Fix MSVC static analysis warning introduced by Dev-VR.

diff --git a/Engine/Plugins/Runtime/Oculus/OculusVR/Source/OculusHMD/Private/OculusBoundaryComponent.cpp b/Engine/Plugins/Runtime/Oculus/OculusVR/Source/OculusHMD/Private/OculusBoundaryComponent.cpp
--- a/Engine/Plugins/Runtime/Oculus/OculusVR/Source/OculusHMD/Private/OculusBoundaryComponent.cpp
+++ b/Engine/Plugins/Runtime/Oculus/OculusVR/Source/OculusHMD/Private/OculusBoundaryComponent.cpp
@@ -132,12 +132,13 @@ static TArray<FVector> GetBoundaryPoints(ovrpBoundaryType BoundaryType)
 	if (OVRP_SUCCESS(ovrp_GetBoundaryGeometry3(BoundaryType, NULL, &NumPoints)))
 	{
         //allocate points
-		ovrpVector3f* BoundaryPoints = new ovrpVector3f[NumPoints];
 		const int BufferSize = NumPoints;
+		ovrpVector3f* BoundaryPoints = new ovrpVector3f[BufferSize];
 	
         if (OVRP_SUCCESS(ovrp_GetBoundaryGeometry3(BoundaryType, BoundaryPoints, &NumPoints)))
         {
 			NumPoints = FMath::Min(BufferSize, NumPoints);
+			check(NumPoints <= BufferSize); // For static analyzer
 			BoundaryPointList.Reserve(NumPoints);
 
             for (int i = 0; i < NumPoints; i++)
