Disable OculusrVR plugin on Android for now; not sure if it's meant to work, but executable fails to link in CIS, and OCULUS_HMD_SUPPORTED_PLATFORMS not being defined causes a bunch of warnings.

diff --git a/Engine/Plugins/Runtime/Oculus/OculusVR/OculusVR.uplugin b/Engine/Plugins/Runtime/Oculus/OculusVR/OculusVR.uplugin
--- a/Engine/Plugins/Runtime/Oculus/OculusVR/OculusVR.uplugin
+++ b/Engine/Plugins/Runtime/Oculus/OculusVR/OculusVR.uplugin
@@ -20,12 +20,12 @@
 			"Name" : "OculusHMD",
 			"Type" : "Runtime",
 			"LoadingPhase" : "PostConfigInit",
-			"WhitelistPlatforms" : [ "Win64", "Win32", "Android" ]
+			"WhitelistPlatforms" : [ "Win64", "Win32" ]
 		},
 		{
 			"Name" : "OculusInput",
 			"Type" : "Runtime",
-			"WhitelistPlatforms" : [ "Win64", "Win32", "Android" ]
+			"WhitelistPlatforms" : [ "Win64", "Win32" ]
 		},
 	]
 }
