Restoring Android support for OculusHMD plugin.  Wrapped platform specific defines in OculusHMDPrivateRHI.h with OCULUS_HMD_SUPPORTED_PLATFORMS, which prevents undefined token warnings in Clang (but not on VS)

diff --git a/Engine/Plugins/Runtime/Oculus/OculusVR/OculusVR.uplugin b/Engine/Plugins/Runtime/Oculus/OculusVR/OculusVR.uplugin
--- a/Engine/Plugins/Runtime/Oculus/OculusVR/OculusVR.uplugin
+++ b/Engine/Plugins/Runtime/Oculus/OculusVR/OculusVR.uplugin
@@ -20,12 +20,12 @@
 			"Name" : "OculusHMD",
 			"Type" : "Runtime",
 			"LoadingPhase" : "PostConfigInit",
-			"WhitelistPlatforms" : [ "Win64", "Win32" ]
+			"WhitelistPlatforms" : [ "Win64", "Win32", "Android" ]
 		},
 		{
 			"Name" : "OculusInput",
 			"Type" : "Runtime",
-			"WhitelistPlatforms" : [ "Win64", "Win32" ]
+			"WhitelistPlatforms" : [ "Win64", "Win32", "Android" ]
 		},
 	]
 }
diff --git a/Engine/Plugins/Runtime/Oculus/OculusVR/Source/OculusHMD/Private/OculusHMDPrivateRHI.h b/Engine/Plugins/Runtime/Oculus/OculusVR/Source/OculusHMD/Private/OculusHMDPrivateRHI.h
--- a/Engine/Plugins/Runtime/Oculus/OculusVR/Source/OculusHMD/Private/OculusHMDPrivateRHI.h
+++ b/Engine/Plugins/Runtime/Oculus/OculusVR/Source/OculusHMD/Private/OculusHMDPrivateRHI.h
@@ -2,6 +2,8 @@
 
 #pragma once
 
+#if OCULUS_HMD_SUPPORTED_PLATFORMS
+
 //-------------------------------------------------------------------------------------------------
 // D3D11
 //-------------------------------------------------------------------------------------------------
@@ -72,3 +74,5 @@
 #include "VulkanRHIPrivate.h"
 #include "VulkanResources.h"
 #endif // OCULUS_HMD_SUPPORTED_PLATFORMS_VULKAN
+
+#endif // OCULUS_HMD_SUPPORTED_PLATFORMS
\ No newline at end of file
