Fix more static analysis warnings in the same vein as CL 3513462.

diff --git a/Engine/Plugins/Runtime/Oculus/OculusVR/Source/OculusHMD/Private/OculusHMDPrivate.cpp b/Engine/Plugins/Runtime/Oculus/OculusVR/Source/OculusHMD/Private/OculusHMDPrivate.cpp
--- a/Engine/Plugins/Runtime/Oculus/OculusVR/Source/OculusHMD/Private/OculusHMDPrivate.cpp
+++ b/Engine/Plugins/Runtime/Oculus/OculusVR/Source/OculusHMD/Private/OculusHMDPrivate.cpp
@@ -97,9 +97,9 @@ void ExecuteOnRenderThread_DoNotWait(const std::function<void()>& Function)
 	{
 		ENQUEUE_UNIQUE_RENDER_COMMAND_ONEPARAMETER(
 			ExecuteOnRenderThread,
-			std::function<void()>, function, Function,
+			std::function<void()>, FunctionArg, Function,
 			{
-				function();
+				FunctionArg();
 			});
 	}
 	else
@@ -117,9 +117,9 @@ void ExecuteOnRenderThread(const std::function<void(FRHICommandListImmediate&)>&
 	{
 		ENQUEUE_UNIQUE_RENDER_COMMAND_ONEPARAMETER(
 			ExecuteOnRenderThread,
-			std::function<void(FRHICommandListImmediate&)>, function, Function,
+			std::function<void(FRHICommandListImmediate&)>, FunctionArg, Function,
 			{
-				function(RHICmdList);
+				FunctionArg(RHICmdList);
 			});
 
 		FlushRenderingCommands();
@@ -139,9 +139,9 @@ void ExecuteOnRenderThread_DoNotWait(const std::function<void(FRHICommandListImm
 	{
 		ENQUEUE_UNIQUE_RENDER_COMMAND_ONEPARAMETER(
 			ExecuteOnRenderThread,
-			std::function<void(FRHICommandListImmediate&)>, function, Function,
+			std::function<void(FRHICommandListImmediate&)>, FunctionArg, Function,
 			{
-				function(RHICmdList);
+				FunctionArg(RHICmdList);
 			});
 	}
 	else
