Defaulting OCULUS_HMD_... defines to 0 on platforms that oculus doesn't support (so that files that use these defines instead as gates/wrappers don't warn).

diff --git a/Engine/Plugins/Runtime/Oculus/OculusVR/Source/OculusHMD/Private/OculusHMDPrivate.h b/Engine/Plugins/Runtime/Oculus/OculusVR/Source/OculusHMD/Private/OculusHMDPrivate.h
--- a/Engine/Plugins/Runtime/Oculus/OculusVR/Source/OculusHMD/Private/OculusHMDPrivate.h
+++ b/Engine/Plugins/Runtime/Oculus/OculusVR/Source/OculusHMD/Private/OculusHMDPrivate.h
@@ -13,6 +13,11 @@
 #define OCULUS_HMD_SUPPORTED_PLATFORMS_D3D12 PLATFORM_WINDOWS
 #define OCULUS_HMD_SUPPORTED_PLATFORMS_OPENGL (PLATFORM_WINDOWS || PLATFORM_ANDROID)
 #define OCULUS_HMD_SUPPORTED_PLATFORMS_VULKAN (PLATFORM_WINDOWS || PLATFORM_ANDROID)
+#else 
+#define OCULUS_HMD_SUPPORTED_PLATFORMS_D3D11 0
+#define OCULUS_HMD_SUPPORTED_PLATFORMS_D3D12 0
+#define OCULUS_HMD_SUPPORTED_PLATFORMS_OPENGL 0
+#define OCULUS_HMD_SUPPORTED_PLATFORMS_VULKAN 0
 #endif // OCULUS_HMD_SUPPORTED_PLATFORMS
 
 
