Use the same Vulkan SDK version check that VulkanRHI uses to ensure we don't fall back to an out of date sdk.

diff --git a/Engine/Plugins/Runtime/Oculus/OculusVR/Source/OculusHMD/OculusHMD.Build.cs b/Engine/Plugins/Runtime/Oculus/OculusVR/Source/OculusHMD/OculusHMD.Build.cs
--- a/Engine/Plugins/Runtime/Oculus/OculusVR/Source/OculusHMD/OculusHMD.Build.cs
+++ b/Engine/Plugins/Runtime/Oculus/OculusVR/Source/OculusHMD/OculusHMD.Build.cs
@@ -79,30 +79,37 @@ namespace UnrealBuildTool.Rules
 
 				// Vulkan
 				{
-					string VulkanSDKPath = Environment.GetEnvironmentVariable("VULKAN_SDK");
-					if (!String.IsNullOrEmpty(VulkanSDKPath))
-					{
-						// If the user has an installed SDK, use that instead
-						PrivateIncludePaths.Add(VulkanSDKPath + "/Include");
-						// Older SDKs have an extra subfolder
-						PrivateIncludePaths.Add(VulkanSDKPath + "/Include/vulkan");
-
-						if (Target.Platform == UnrealTargetPlatform.Win32)
-						{
-							PublicLibraryPaths.Add(VulkanSDKPath + "/Source/lib32");
-						}
-						else
-						{
-							PublicLibraryPaths.Add(VulkanSDKPath + "/Source/lib");
-						}
-
-						PublicAdditionalLibraries.Add("vulkan-1.lib");
-						PublicAdditionalLibraries.Add("vkstatic.1.lib");
-					}
-					else
-					{
-						AddEngineThirdPartyPrivateStaticDependencies(Target, "Vulkan");
-					}
+                    string VulkanSDKPath = Environment.GetEnvironmentVariable("VULKAN_SDK");
+                    bool bSDKInstalled = !String.IsNullOrEmpty(VulkanSDKPath);
+                    bool bUseThirdParty = true;
+
+                    // Check if the installed SDK is newer or the same than the provided headers distributed with the Engine
+                    int ThirdPartyVersion = GetThirdPartyVersion();
+                    int SDKVersion = GetSDKVersion(VulkanSDKPath);
+                    if (SDKVersion >= ThirdPartyVersion)
+                    {
+                        // If the user has an installed SDK, use that instead
+                        PrivateIncludePaths.Add(VulkanSDKPath + "/Include");
+                        // Older SDKs have an extra subfolder
+                        PrivateIncludePaths.Add(VulkanSDKPath + "/Include/vulkan");
+
+                        if (Target.Platform == UnrealTargetPlatform.Win32)
+                        {
+                            PublicLibraryPaths.Add(VulkanSDKPath + "/Source/lib32");
+                        }
+                        else
+                        {
+                            PublicLibraryPaths.Add(VulkanSDKPath + "/Source/lib");
+                        }
+
+                        PublicAdditionalLibraries.Add("vulkan-1.lib");
+                        PublicAdditionalLibraries.Add("vkstatic.1.lib");
+                        bUseThirdParty = false;
+                    }
+                    if (bUseThirdParty)
+                    {
+                        AddEngineThirdPartyPrivateStaticDependencies(Target, "Vulkan");
+                    }
 				}
 
 				// OVRPlugin
@@ -148,5 +155,57 @@ namespace UnrealBuildTool.Rules
 				}
 			}
 		}
-	}
+
+        static int GetVersionFromString(string Text)
+        {
+            string Token = "#define VK_HEADER_VERSION ";
+            Int32 FoundIndex = Text.IndexOf(Token);
+            if (FoundIndex > 0)
+            {
+                string Version = Text.Substring(FoundIndex + Token.Length, 5);
+                int Index = 0;
+                while (Version[Index] >= '0' && Version[Index] <= '9')
+                {
+                    ++Index;
+                }
+
+                Version = Version.Substring(0, Index);
+
+                int VersionNumber = Convert.ToInt32(Version);
+                return VersionNumber;
+            }
+
+            return -1;
+        }
+        static int GetThirdPartyVersion()
+        {
+            try
+            {
+                // Extract current version on ThirdParty
+                string Text = File.ReadAllText("ThirdParty/Vulkan/Windows/Include/vulkan/vulkan.h");
+                return GetVersionFromString(Text);
+            }
+            catch (Exception)
+            {
+            }
+
+            return -1;
+        }
+
+        static int GetSDKVersion(string VulkanSDKPath)
+        {
+            try
+            {
+                // Extract current version on the SDK folder
+                string Header = Path.Combine(VulkanSDKPath, "Include/vulkan/vulkan.h");
+                string Text = File.ReadAllText(Header);
+                return GetVersionFromString(Text);
+            }
+            catch (Exception)
+            {
+            }
+
+            return -1;
+        }
+    }
 }
\ No newline at end of file
