Bad merge

diff --git a/Engine/Plugins/Runtime/Oculus/OculusVR/Source/OculusHMD/OculusHMD.Build.cs b/Engine/Plugins/Runtime/Oculus/OculusVR/Source/OculusHMD/OculusHMD.Build.cs
--- a/Engine/Plugins/Runtime/Oculus/OculusVR/Source/OculusHMD/OculusHMD.Build.cs
+++ b/Engine/Plugins/Runtime/Oculus/OculusVR/Source/OculusHMD/OculusHMD.Build.cs
@@ -82,29 +82,31 @@ namespace UnrealBuildTool.Rules
                     string VulkanSDKPath = Environment.GetEnvironmentVariable("VULKAN_SDK");
                     bool bSDKInstalled = !String.IsNullOrEmpty(VulkanSDKPath);
                     bool bUseThirdParty = true;
-
-                    // Check if the installed SDK is newer or the same than the provided headers distributed with the Engine
-                    int ThirdPartyVersion = GetThirdPartyVersion();
-                    int SDKVersion = GetSDKVersion(VulkanSDKPath);
-                    if (SDKVersion >= ThirdPartyVersion)
+                    if (bSDKInstalled)
                     {
-                        // If the user has an installed SDK, use that instead
-                        PrivateIncludePaths.Add(VulkanSDKPath + "/Include");
-                        // Older SDKs have an extra subfolder
-                        PrivateIncludePaths.Add(VulkanSDKPath + "/Include/vulkan");
-
-                        if (Target.Platform == UnrealTargetPlatform.Win32)
+                        // Check if the installed SDK is newer or the same than the provided headers distributed with the Engine
+                        int ThirdPartyVersion = GetThirdPartyVersion();
+                        int SDKVersion = GetSDKVersion(VulkanSDKPath);
+                        if (SDKVersion >= ThirdPartyVersion)
                         {
-                            PublicLibraryPaths.Add(VulkanSDKPath + "/Source/lib32");
+                            // If the user has an installed SDK, use that instead
+                            PrivateIncludePaths.Add(VulkanSDKPath + "/Include");
+                            // Older SDKs have an extra subfolder
+                            PrivateIncludePaths.Add(VulkanSDKPath + "/Include/vulkan");
+
+                            if (Target.Platform == UnrealTargetPlatform.Win32)
+                            {
+                                PublicLibraryPaths.Add(VulkanSDKPath + "/Source/lib32");
+                            }
+                            else
+                            {
+                                PublicLibraryPaths.Add(VulkanSDKPath + "/Source/lib");
+                            }
+
+                            PublicAdditionalLibraries.Add("vulkan-1.lib");
+                            PublicAdditionalLibraries.Add("vkstatic.1.lib");
+                            bUseThirdParty = false;
                         }
-                        else
-                        {
-                            PublicLibraryPaths.Add(VulkanSDKPath + "/Source/lib");
-                        }
-
-                        PublicAdditionalLibraries.Add("vulkan-1.lib");
-                        PublicAdditionalLibraries.Add("vkstatic.1.lib");
-                        bUseThirdParty = false;
                     }
                     if (bUseThirdParty)
                     {
