Disable rhi thread when running with vulkan on oculus android platforms.

diff --git a/Engine/Plugins/Runtime/Oculus/OculusVR/Source/OculusHMD/Private/OculusHMD_CustomPresent_Vulkan.cpp b/Engine/Plugins/Runtime/Oculus/OculusVR/Source/OculusHMD/Private/OculusHMD_CustomPresent_Vulkan.cpp
--- a/Engine/Plugins/Runtime/Oculus/OculusVR/Source/OculusHMD/Private/OculusHMD_CustomPresent_Vulkan.cpp
+++ b/Engine/Plugins/Runtime/Oculus/OculusVR/Source/OculusHMD/Private/OculusHMD_CustomPresent_Vulkan.cpp
@@ -1,4 +1,4 @@
-// Copyright 1998-2018 Epic Games, Inc. All Rights Reserved.
+﻿// Copyright 1998-2018 Epic Games, Inc. All Rights Reserved.
 
 #include "OculusHMD_CustomPresent.h"
 #include "OculusHMDPrivateRHI.h"
@@ -38,6 +38,12 @@ public:
 FVulkanCustomPresent::FVulkanCustomPresent(FOculusHMD* InOculusHMD) :
 	FCustomPresent(InOculusHMD, ovrpRenderAPI_Vulkan, PF_R8G8B8A8, false, false)
 {
+#if PLATFORM_ANDROID
+	if (GRHISupportsRHIThread && GIsThreadedRendering && GUseRHIThread_InternalUseOnly)
+	{
+		SetRHIThreadEnabled(false, false);
+	}
+#endif
 }
 
 
