Mirroring CL 4081823 from Dev-VR

Loading Oculus controller models in the editor so that they get included in the cooking process (so that the MotionController components that indirectly uses these at runtime don't get dropped).

diff --git a/Engine/Plugins/Runtime/Oculus/OculusVR/Source/OculusEditor/Private/OculusEditorModule.cpp b/Engine/Plugins/Runtime/Oculus/OculusVR/Source/OculusEditor/Private/OculusEditorModule.cpp
--- a/Engine/Plugins/Runtime/Oculus/OculusVR/Source/OculusEditor/Private/OculusEditorModule.cpp
+++ b/Engine/Plugins/Runtime/Oculus/OculusVR/Source/OculusEditor/Private/OculusEditorModule.cpp
@@ -5,6 +5,7 @@
 #include "Modules/ModuleManager.h"
 #include "UObject/WeakObjectPtr.h"
 #include "UObject/Class.h"
+#include "OculusAssetDirectory.h"
 
 // Settings
 #include "ISettingsModule.h"
@@ -20,10 +21,13 @@ public:
 	virtual void StartupModule() override
 	{				
 		RegisterSettings();
+		FOculusAssetDirectory::LoadForCook();
 	}
 
 	virtual void ShutdownModule() override
 	{
+		FOculusAssetDirectory::ReleaseAll();
+
 		if (UObjectInitialized())
 		{
 			UnregisterSettings();  
diff --git a/Engine/Plugins/Runtime/Oculus/OculusVR/Source/OculusHMD/Private/OculusAssetManager.cpp b/Engine/Plugins/Runtime/Oculus/OculusVR/Source/OculusHMD/Private/OculusAssetManager.cpp
--- a/Engine/Plugins/Runtime/Oculus/OculusVR/Source/OculusHMD/Private/OculusAssetManager.cpp
+++ b/Engine/Plugins/Runtime/Oculus/OculusVR/Source/OculusHMD/Private/OculusAssetManager.cpp
@@ -7,6 +7,64 @@
 #include "UObject/SoftObjectPath.h"
 #include "Engine/SkeletalMesh.h"
 #include "Components/SkeletalMeshComponent.h"
+#include "OculusAssetDirectory.h"
+#include "UObject/GCObject.h"
+
+/* FOculusAssetDirectory
+ *****************************************************************************/
+
+FSoftObjectPath FOculusAssetDirectory::AssetListing[] =
+{
+	FString(TEXT("/OculusVR/Meshes/RiftHMD.RiftHMD")),
+	FString(TEXT("/OculusVR/Meshes/GearVRController.GearVRController")),
+	FString(TEXT("/OculusVR/Meshes/LeftTouchController.LeftTouchController")),
+	FString(TEXT("/OculusVR/Meshes/RightTouchController.RightTouchController")),
+};
+
+#if WITH_EDITORONLY_DATA
+class FOculusAssetRepo : public FGCObject, public TArray<UObject*>
+{
+public:
+	// made an on-demand singleton rather than a static global, to avoid issues with FGCObject initialization
+	static FOculusAssetRepo& Get()
+	{
+		static FOculusAssetRepo AssetRepository;
+		return AssetRepository;
+	}
+
+	UObject* LoadAndAdd(const FSoftObjectPath& AssetPath)
+	{
+		UObject* AssetObj = AssetPath.TryLoad();
+		if (AssetObj != nullptr)
+		{
+			AddUnique(AssetObj);
+		}
+		return AssetObj;
+	}
+
+public:
+	//~ FGCObject interface
+	virtual void AddReferencedObjects(FReferenceCollector& Collector) override
+	{
+		Collector.AddReferencedObjects(*this);
+	}
+};
+
+void FOculusAssetDirectory::LoadForCook()
+{
+	FOculusAssetRepo& AssetRepro = FOculusAssetRepo::Get();
+	for (int32 AssetIndex = 0; AssetIndex < ARRAY_COUNT(FOculusAssetDirectory::AssetListing); ++AssetIndex)
+	{
+		AssetRepro.LoadAndAdd(FOculusAssetDirectory::AssetListing[AssetIndex]);
+	}
+}
+
+void FOculusAssetDirectory::ReleaseAll()
+{
+	FOculusAssetRepo::Get().Empty();
+}
+#endif // WITH_EDITORONLY_DATA
+
 
 /* OculusAssetManager_Impl
  *****************************************************************************/
@@ -22,13 +80,13 @@ namespace OculusAssetManager_Impl
 
 	static FRenderableDevice RenderableDevices[] =
 	{
-		{ ovrpNode_Head,      FString(TEXT("/OculusVR/Meshes/RiftHMD.RiftHMD")) },
+		{ ovrpNode_Head,      FOculusAssetDirectory::AssetListing[0] },
 #if PLATFORM_ANDROID
-		{ ovrpNode_HandLeft,  FString(TEXT("/OculusVR/Meshes/GearVRController.GearVRController")) },
-		{ ovrpNode_HandRight, FString(TEXT("/OculusVR/Meshes/GearVRController.GearVRController")) },
+		{ ovrpNode_HandLeft,  FOculusAssetDirectory::AssetListing[1] },
+		{ ovrpNode_HandRight, FOculusAssetDirectory::AssetListing[1] },
 #else 
-		{ ovrpNode_HandLeft,  FString(TEXT("/OculusVR/Meshes/LeftTouchController.LeftTouchController")) },
-		{ ovrpNode_HandRight, FString(TEXT("/OculusVR/Meshes/RightTouchController.RightTouchController")) },
+		{ ovrpNode_HandLeft,  FOculusAssetDirectory::AssetListing[2] },
+		{ ovrpNode_HandRight, FOculusAssetDirectory::AssetListing[3] },
 #endif
 	};
 
@@ -164,3 +222,4 @@ UPrimitiveComponent* FOculusAssetManager::CreateRenderComponent(const int32 Devi
 	OnLoadComplete.ExecuteIfBound(NewRenderComponent);
 	return NewRenderComponent;
 }
+
diff --git a/Engine/Plugins/Runtime/Oculus/OculusVR/Source/OculusHMD/Public/OculusAssetDirectory.h b/Engine/Plugins/Runtime/Oculus/OculusVR/Source/OculusHMD/Public/OculusAssetDirectory.h
new file mode 100644
--- /dev/null
+++ b/Engine/Plugins/Runtime/Oculus/OculusVR/Source/OculusHMD/Public/OculusAssetDirectory.h
@@ -0,0 +1,17 @@
+// Copyright 1998-2018 Epic Games, Inc. All Rights Reserved.
+
+#pragma once
+
+#include "CoreMinimal.h"
+#include "UObject/SoftObjectPath.h"
+
+class FOculusAssetDirectory
+{
+public: 
+#if WITH_EDITORONLY_DATA
+	OCULUSHMD_API static void LoadForCook();
+	OCULUSHMD_API static void ReleaseAll();
+#endif
+
+	static FSoftObjectPath AssetListing[];
+};
