Mirroring CL 4083160 from Dev-VR

Making it so the HMD post-process pass properly fills out its RenderTargetDesc (so it doesn't rely on being the final output to have it set).

diff --git a/Engine/Plugins/Runtime/Oculus/OculusVR/Source/OculusHMD/Private/OculusAssetManager.cpp b/Engine/Plugins/Runtime/Oculus/OculusVR/Source/OculusHMD/Private/OculusAssetManager.cpp
--- a/Engine/Plugins/Runtime/Oculus/OculusVR/Source/OculusHMD/Private/OculusAssetManager.cpp
+++ b/Engine/Plugins/Runtime/Oculus/OculusVR/Source/OculusHMD/Private/OculusAssetManager.cpp
@@ -19,6 +19,8 @@ FSoftObjectPath FOculusAssetDirectory::AssetListing[] =
 	FString(TEXT("/OculusVR/Meshes/GearVRController.GearVRController")),
 	FString(TEXT("/OculusVR/Meshes/LeftTouchController.LeftTouchController")),
 	FString(TEXT("/OculusVR/Meshes/RightTouchController.RightTouchController")),
+
+	FString(TEXT("/OculusVR/Materials/PokeAHoleMaterial.PokeAHoleMaterial"))
 };
 
 #if WITH_EDITORONLY_DATA
