Mirroring CL 4088048 from Dev-VR

Use the standardized tracking-to-world transform for world layer positioning on Oculus - handles both explicit and implicit HMD positioning.

diff --git a/Engine/Plugins/Runtime/Oculus/OculusVR/Source/OculusHMD/Private/OculusHMD.cpp b/Engine/Plugins/Runtime/Oculus/OculusVR/Source/OculusHMD/Private/OculusHMD.cpp
--- a/Engine/Plugins/Runtime/Oculus/OculusVR/Source/OculusHMD/Private/OculusHMD.cpp
+++ b/Engine/Plugins/Runtime/Oculus/OculusVR/Source/OculusHMD/Private/OculusHMD.cpp
@@ -935,7 +935,17 @@ namespace OculusHMD
 
 		FGameFrame* const CurrentGameFrame = Frame.Get();
 
-		if (!InWorldContext.World() || (!(GEnableVREditorHacks && InWorldContext.WorldType == EWorldType::Editor) && !InWorldContext.World()->IsGameWorld()) || !CurrentGameFrame)
+		if (CurrentGameFrame)
+		{
+			// don't use the cached value, as it could be affected by the player's position, so we update it here at the latest point in the game frame
+			CurrentGameFrame->TrackingToWorld = ComputeTrackingToWorldTransform(InWorldContext);
+		}
+		else
+		{
+			return false;
+		}
+
+		if ( !InWorldContext.World() || (!(GEnableVREditorHacks && InWorldContext.WorldType == EWorldType::Editor) && !InWorldContext.World()->IsGameWorld()) )
 		{
 			// ignore all non-game worlds
 			return false;
diff --git a/Engine/Plugins/Runtime/Oculus/OculusVR/Source/OculusHMD/Private/OculusHMD_GameFrame.h b/Engine/Plugins/Runtime/Oculus/OculusVR/Source/OculusHMD/Private/OculusHMD_GameFrame.h
--- a/Engine/Plugins/Runtime/Oculus/OculusVR/Source/OculusHMD/Private/OculusHMD_GameFrame.h
+++ b/Engine/Plugins/Runtime/Oculus/OculusVR/Source/OculusHMD/Private/OculusHMD_GameFrame.h
@@ -27,6 +27,8 @@ public:
 	FVector PlayerLocation;			// (CalculateStereoViewOffset)
 	float NearClippingPlane;		// (GetStereoProjectionMatrix)
 
+	FTransform TrackingToWorld;		// (OnEndGameFrame)
+
 	ETiledMultiResLevel MultiResLevel; // OnStartGameFrame
 
 	FIntRect FinalViewRect[3]; // SetFinalViewRect
diff --git a/Engine/Plugins/Runtime/Oculus/OculusVR/Source/OculusHMD/Private/OculusHMD_Layer.cpp b/Engine/Plugins/Runtime/Oculus/OculusVR/Source/OculusHMD/Private/OculusHMD_Layer.cpp
--- a/Engine/Plugins/Runtime/Oculus/OculusVR/Source/OculusHMD/Private/OculusHMD_Layer.cpp
+++ b/Engine/Plugins/Runtime/Oculus/OculusVR/Source/OculusHMD/Private/OculusHMD_Layer.cpp
@@ -465,8 +465,8 @@ const ovrpLayerSubmit* FLayer::UpdateLayer_RHIThread(const FSettings* Settings,
 		switch (Desc.PositionType)
 		{
 		case IStereoLayers::WorldLocked:
-			BaseOrientation = Frame->PlayerOrientation;
-			BaseLocation = Frame->PlayerLocation;
+			BaseOrientation = Frame->TrackingToWorld.GetRotation();
+			BaseLocation = Frame->TrackingToWorld.GetTranslation();
 			break;
 
 		case IStereoLayers::TrackerLocked:
@@ -480,13 +480,13 @@ const ovrpLayerSubmit* FLayer::UpdateLayer_RHIThread(const FSettings* Settings,
 			break;
 		}
 
-		FTransform playerTransform(BaseOrientation, BaseLocation);
+		FTransform PlayerTransform(BaseOrientation, BaseLocation);
 
 		FQuat Orientation = Desc.Transform.Rotator().Quaternion();
 		FVector Location = Desc.Transform.GetLocation();
 
 		OvrpLayerSubmit.Pose.Orientation = ToOvrpQuatf(BaseOrientation.Inverse() * Orientation);
-		OvrpLayerSubmit.Pose.Position = ToOvrpVector3f((playerTransform.InverseTransformPosition(Location)) * LocationScale);
+		OvrpLayerSubmit.Pose.Position = ToOvrpVector3f((PlayerTransform.InverseTransformPosition(Location)) * LocationScale);
 		OvrpLayerSubmit.LayerSubmitFlags = 0;
 
 		if (Desc.PositionType == IStereoLayers::FaceLocked)
