Making sure to load the third party dll, and initing Oculus when we querry for Vulkan extensions. Otherwise we end up calling into the OVR plugin before the dll is loaded.

diff --git a/Engine/Plugins/Runtime/Oculus/OculusVR/Source/OculusHMD/Private/OculusHMDModule.cpp b/Engine/Plugins/Runtime/Oculus/OculusVR/Source/OculusHMD/Private/OculusHMDModule.cpp
--- a/Engine/Plugins/Runtime/Oculus/OculusVR/Source/OculusHMD/Private/OculusHMDModule.cpp
+++ b/Engine/Plugins/Runtime/Oculus/OculusVR/Source/OculusHMD/Private/OculusHMDModule.cpp
@@ -215,11 +215,13 @@ TSharedPtr< class IXRTrackingSystem, ESPMode::ThreadSafe > FOculusHMDModule::Cre
 TSharedPtr< IHeadMountedDisplayVulkanExtensions, ESPMode::ThreadSafe >  FOculusHMDModule::GetVulkanExtensions()
 {
 #if OCULUS_HMD_SUPPORTED_PLATFORMS
-	if (!VulkanExtensions.IsValid())
+	if (PreInit())
 	{
-		VulkanExtensions = MakeShareable(new OculusHMD::FVulkanExtensions);
+		if (!VulkanExtensions.IsValid())
+		{
+			VulkanExtensions = MakeShareable(new OculusHMD::FVulkanExtensions);
+		}
 	}
-
 	return VulkanExtensions;
 #endif
 	return nullptr;
