Guarding against bad settings pointers internal to Oculus.

diff --git a/Engine/Plugins/Runtime/Oculus/OculusVR/Source/OculusHMD/Private/OculusHMD.cpp b/Engine/Plugins/Runtime/Oculus/OculusVR/Source/OculusHMD/Private/OculusHMD.cpp
--- a/Engine/Plugins/Runtime/Oculus/OculusVR/Source/OculusHMD/Private/OculusHMD.cpp
+++ b/Engine/Plugins/Runtime/Oculus/OculusVR/Source/OculusHMD/Private/OculusHMD.cpp
@@ -1427,7 +1427,7 @@ namespace OculusHMD
 	}
 
 	FIntRect FOculusHMD::GetFullFlatEyeRect_RenderThread(FTexture2DRHIRef EyeTexture) const
-			{
+	{
 		check(IsInRenderingThread());
 		// Rift does this differently than other platforms, it already has an idea of what rectangle it wants to use stored.
 		FIntRect& EyeRect = Settings_RenderThread->EyeRenderViewport[0];
@@ -1498,7 +1498,8 @@ namespace OculusHMD
 	{
 		CheckInGameThread();
 
-		return Settings->IsStereoEnabled() && bNeedReAllocateViewportRenderTarget;
+		return ensureMsgf(Settings.IsValid(), TEXT("Unexpected issue with Oculus settings on the GameThread. This should be valid when this is called in EnqueueBeginRenderFrame() - has the callsite changed?")) &&
+			Settings->IsStereoEnabled() && bNeedReAllocateViewportRenderTarget;
 	}
 
 
@@ -1506,7 +1507,8 @@ namespace OculusHMD
 	{
 		CheckInRenderThread();
 
-		return Settings_RenderThread->IsStereoEnabled() && bNeedReAllocateDepthTexture_RenderThread;
+		return ensureMsgf(Settings_RenderThread.IsValid(), TEXT("Unexpected issue with Oculus settings on the RenderThread. This should be valid when this is called in AllocateCommonDepthTargets() - has the callsite changed?")) &&
+			Settings_RenderThread->IsStereoEnabled() && bNeedReAllocateDepthTexture_RenderThread;
 	}
 
 
