Fix UE4Editor Linux debug crash with -vulkan

diff --git a/Engine/Plugins/Runtime/Oculus/OculusVR/Source/OculusHMD/OculusHMD.Build.cs b/Engine/Plugins/Runtime/Oculus/OculusVR/Source/OculusHMD/OculusHMD.Build.cs
--- a/Engine/Plugins/Runtime/Oculus/OculusVR/Source/OculusHMD/OculusHMD.Build.cs
+++ b/Engine/Plugins/Runtime/Oculus/OculusVR/Source/OculusHMD/OculusHMD.Build.cs
@@ -24,21 +24,13 @@ namespace UnrealBuildTool.Rules
 					"ProceduralMeshComponent",
 				});			
 
-			if (Target.Platform == UnrealTargetPlatform.Android)
-			{
-				PrivateIncludePaths.Add("../../../../../Source/Runtime/VulkanRHI/Private/Android");
-			}
-			else if (Target.Platform == UnrealTargetPlatform.Linux)
-			{
-				PrivateIncludePaths.Add("../../../../../Source/Runtime/VulkanRHI/Private/Linux");
-			}
-			else if (Target.Platform == UnrealTargetPlatform.Quail)
+			if (Target.Platform == UnrealTargetPlatform.Win32 || Target.Platform == UnrealTargetPlatform.Win64)
 			{
-				PrivateIncludePaths.Add("../../../../../Source/Runtime/VulkanRHI/Private/Quail");
+				PrivateIncludePaths.Add("../../../../../Source/Runtime/VulkanRHI/Private/Windows");
 			}
-			else if (Target.Platform == UnrealTargetPlatform.Win32 || Target.Platform == UnrealTargetPlatform.Win64)
+			else
 			{
-				PrivateIncludePaths.Add("../../../../../Source/Runtime/VulkanRHI/Private/Windows");
+				PrivateIncludePaths.Add("../../../../../Source/Runtime/VulkanRHI/Private/" + Target.Platform);
 			}
 
 
@@ -238,4 +230,4 @@ namespace UnrealBuildTool.Rules
 			return -1;
 		}
 	}
-}
\ No newline at end of file
+}
diff --git a/Engine/Plugins/Runtime/Steam/SteamVR/Source/SteamVR/SteamVR.Build.cs b/Engine/Plugins/Runtime/Steam/SteamVR/Source/SteamVR/SteamVR.Build.cs
--- a/Engine/Plugins/Runtime/Steam/SteamVR/Source/SteamVR/SteamVR.Build.cs
+++ b/Engine/Plugins/Runtime/Steam/SteamVR/Source/SteamVR/SteamVR.Build.cs
@@ -15,21 +15,13 @@ namespace UnrealBuildTool.Rules
 				}
 				);
 
-			if(Target.Platform == UnrealTargetPlatform.Android)
+			if(Target.Platform == UnrealTargetPlatform.Win32 || Target.Platform == UnrealTargetPlatform.Win64)
 			{
-				PrivateIncludePaths.Add("../../../../../Source/Runtime/VulkanRHI/Private/Android");
-			}
-			else if(Target.Platform == UnrealTargetPlatform.Linux)
-			{
-				PrivateIncludePaths.Add("../../../../../Source/Runtime/VulkanRHI/Private/Linux");
-			}
-			else if(Target.Platform == UnrealTargetPlatform.Quail)
-			{
-				PrivateIncludePaths.Add("../../../../../Source/Runtime/VulkanRHI/Private/Quail");
+				PrivateIncludePaths.Add("../../../../../Source/Runtime/VulkanRHI/Private/Windows");
 			}
-			else if(Target.Platform == UnrealTargetPlatform.Win32 || Target.Platform == UnrealTargetPlatform.Win64)
+			else
 			{
-				PrivateIncludePaths.Add("../../../../../Source/Runtime/VulkanRHI/Private/Windows");
+				PrivateIncludePaths.Add("../../../../../Source/Runtime/VulkanRHI/Private/" + Target.Platform);
 			}
 
 			PrivateDependencyModuleNames.AddRange(
