Mirroring CL 4089560 from Dev-VR

Convert world-based coordinates in UE4 units back to oculus tracking space relative coordinates in meters when checking play area boundaries.

diff --git a/Engine/Plugins/Runtime/Oculus/OculusVR/Source/OculusHMD/Private/OculusBoundaryComponent.cpp b/Engine/Plugins/Runtime/Oculus/OculusVR/Source/OculusHMD/Private/OculusBoundaryComponent.cpp
--- a/Engine/Plugins/Runtime/Oculus/OculusVR/Source/OculusHMD/Private/OculusBoundaryComponent.cpp
+++ b/Engine/Plugins/Runtime/Oculus/OculusVR/Source/OculusHMD/Private/OculusBoundaryComponent.cpp
@@ -88,6 +88,13 @@ static float DistanceToWorldSpace(float ovrDistance)
 	return OculusHMD->ConvertFloat_M2U(ovrDistance);
 }
 
+static ovrpVector3f WorldSpaceToPoint(const FVector& InUnrealVector)
+{
+	FOculusHMD* OculusHMD = (FOculusHMD*)(GEngine->XRSystem->GetHMDDevice());
+	check(OculusHMD);
+	return OculusHMD->WorldLocationToOculusPoint(InUnrealVector);
+}
+
 /**
  * Helper that checks if DeviceType triggers BoundaryType. If so, stores details about interaction and adds to ResultList!
  * @param ResultList Out-parameter -- a list to add to may or may not be specified. Stores each FBoundaryTestResult corresponding to a device that has triggered boundary system
@@ -175,7 +182,7 @@ static FBoundaryTestResult CheckPointInBounds(EBoundaryType BoundaryType, const
 
 	if (FOculusHMDModule::Get().IsOVRPluginAvailable())
 	{
-		ovrpVector3f OvrpPoint = ToOvrpVector3f(Point);
+		ovrpVector3f OvrpPoint = WorldSpaceToPoint(Point);
 		ovrpBoundaryType OvrpBoundaryType = ToOvrpBoundaryType(BoundaryType);
 		ovrpBoundaryTestResult InteractionResult;
 
diff --git a/Engine/Plugins/Runtime/Oculus/OculusVR/Source/OculusHMD/Private/OculusHMD.cpp b/Engine/Plugins/Runtime/Oculus/OculusVR/Source/OculusHMD/Private/OculusHMD.cpp
--- a/Engine/Plugins/Runtime/Oculus/OculusVR/Source/OculusHMD/Private/OculusHMD.cpp
+++ b/Engine/Plugins/Runtime/Oculus/OculusVR/Source/OculusHMD/Private/OculusHMD.cpp
@@ -2979,6 +2979,26 @@ namespace OculusHMD
 		return TransformWithPlayer;
 	}
 
+	ovrpVector3f FOculusHMD::WorldLocationToOculusPoint(const FVector& InUnrealPosition)
+	{
+		CheckInGameThread();
+		FQuat AdjustedPlayerOrientation = GetBaseOrientation().Inverse() * LastPlayerOrientation;
+		AdjustedPlayerOrientation.Normalize();
+
+		FVector AdjustedPlayerLocation = LastPlayerLocation;
+		if (GetXRCamera(HMDDeviceId)->GetUseImplicitHMDPosition())
+		{
+			FQuat HeadOrientation = FQuat::Identity; // Unused
+			FVector HeadPosition;
+			GetCurrentPose(HMDDeviceId, HeadOrientation, HeadPosition);
+			AdjustedPlayerLocation -= LastPlayerOrientation.Inverse().RotateVector(HeadPosition);
+		}
+		const FTransform InvWorldTransform = FTransform(AdjustedPlayerOrientation, AdjustedPlayerLocation).Inverse();
+		const FVector ConvertedPosition = InvWorldTransform.TransformPosition(InUnrealPosition) / GetWorldToMetersScale();
+
+		return ToOvrpVector3f(ConvertedPosition);
+	}
+
 
 	float FOculusHMD::ConvertFloat_M2U(float OculusFloat) const
 	{
diff --git a/Engine/Plugins/Runtime/Oculus/OculusVR/Source/OculusHMD/Private/OculusHMD.h b/Engine/Plugins/Runtime/Oculus/OculusVR/Source/OculusHMD/Private/OculusHMD.h
--- a/Engine/Plugins/Runtime/Oculus/OculusVR/Source/OculusHMD/Private/OculusHMD.h
+++ b/Engine/Plugins/Runtime/Oculus/OculusVR/Source/OculusHMD/Private/OculusHMD.h
@@ -267,6 +267,9 @@ public:
 	/** Turns ovrVector3f in Unreal World space to a scaled FVector and applies translation and rotation corresponding to player movement */
 	FVector ScaleAndMovePointWithPlayer(ovrpVector3f& OculusHMDPoint);
 
+	/** The inverse of ScaleAndMovePointWithPlayer */
+	ovrpVector3f WorldLocationToOculusPoint(const FVector& InUnrealPosition);
+
 	/** Convert dimension of a float (e.g., a distance) from meters to Unreal Units */
 	float ConvertFloat_M2U(float OculusFloat) const;
 	FVector ConvertVector_M2U(ovrpVector3f OculusPoint) const;
