UE-63073 //UE4/Main - Compile ClangEditor Win64 - FInputGesture is deprecated

diff --git a/Engine/Plugins/Runtime/Oculus/OculusVR/Source/OculusEditor/Private/OculusToolCommands.cpp b/Engine/Plugins/Runtime/Oculus/OculusVR/Source/OculusEditor/Private/OculusToolCommands.cpp
--- a/Engine/Plugins/Runtime/Oculus/OculusVR/Source/OculusEditor/Private/OculusToolCommands.cpp
+++ b/Engine/Plugins/Runtime/Oculus/OculusVR/Source/OculusEditor/Private/OculusToolCommands.cpp
@@ -7,7 +7,7 @@
 
 void FOculusToolCommands::RegisterCommands()
 {
-	UI_COMMAND(OpenPluginWindow, "Oculus Tool", "Show Oculus Tool Window", EUserInterfaceActionType::Button, FInputGesture());
+	UI_COMMAND(OpenPluginWindow, "Oculus Tool", "Show Oculus Tool Window", EUserInterfaceActionType::Button, FInputChord());
 }
 
 void FOculusToolCommands::ShowOculusTool()
