Fixes a splash screen related race condition that could result in a crash
SetSplashScreen now overwrites the current splash screen if one is already set, previously it did nothing

diff --git a/Engine/Plugins/Runtime/Oculus/OculusVR/Source/OculusHMD/Private/OculusHMD.cpp b/Engine/Plugins/Runtime/Oculus/OculusVR/Source/OculusHMD/Private/OculusHMD.cpp
--- a/Engine/Plugins/Runtime/Oculus/OculusVR/Source/OculusHMD/Private/OculusHMD.cpp
+++ b/Engine/Plugins/Runtime/Oculus/OculusVR/Source/OculusHMD/Private/OculusHMD.cpp
@@ -1448,24 +1448,37 @@ namespace OculusHMD
 			return;
 		}
 
-		if (bSplashIsShown)
-		{
-			//if update splash screen is shown, update the head orientation default to recenter splash screens
-			FQuat HeadOrientation = FQuat::Identity;
-			FVector HeadPosition;
-			GetCurrentPose(HMDDeviceId, HeadOrientation, HeadPosition);
-			SplashRotation = FRotator(HeadOrientation);
-			SplashRotation.Pitch = 0;
-			SplashRotation.Roll = 0;
+		Flags.bNeedSplashUpdate = true;
+	}
 
-			Splash->Show();
-		}
-		else
+	void FOculusHMD::UpdateSplashScreen_GameThread()
+	{
+		if (Flags.bNeedSplashUpdate)
 		{
-			Splash->Hide();
+			if (bSplashIsShown)
+			{
+				Splash->Show();
+			}
+			else
+			{
+				Splash->Hide();
+			}
+
+			Flags.bNeedSplashUpdate = false;
 		}
 	}
 
+	void FOculusHMD::SetSplashRotationToForward()
+	{
+		//if update splash screen is shown, update the head orientation default to recenter splash screens
+		FQuat HeadOrientation = FQuat::Identity;
+		FVector HeadPosition;
+		GetCurrentPose(HMDDeviceId, HeadOrientation, HeadPosition);
+		SplashRotation = FRotator(HeadOrientation);
+		SplashRotation.Pitch = 0;
+		SplashRotation.Roll = 0;
+	}
+
 	FOculusSplashDesc FOculusHMD::GetUESplashScreenDesc()
 	{
 		FOculusSplashDesc Desc;
@@ -1563,7 +1576,6 @@ namespace OculusHMD
 		return StereoLayerDesc;
 	}
 
-
 	void FOculusHMD::SetupViewFamily(FSceneViewFamily& InViewFamily)
 	{
 		CheckInGameThread();
@@ -1712,6 +1724,7 @@ namespace OculusHMD
 		CachedWorldToMetersScale = 100.0f;
 
 		NextFrameNumber = 1;
+		WaitFrameNumber = 0;
 		NextLayerId = 0;
 
 		Settings = CreateNewSettings();
@@ -2919,6 +2932,7 @@ namespace OculusHMD
 		Result->WorldToMetersScale = CachedWorldToMetersScale;
 		Result->NearClippingPlane = GNearClippingPlane;
 		Result->MultiResLevel = Settings->MultiResLevel;
+		Result->Flags.bSplashIsShown = Splash->IsShown();
 		return Result;
 	}
 
@@ -2930,10 +2944,29 @@ namespace OculusHMD
 
 		if (!Frame.IsValid())
 		{
+			UpdateSplashScreen_GameThread(); //the result of this is used in CreateGameFrame to know if Frame is a "real" one or a "splash" one.
 			Frame = CreateNewGameFrame();
 			NextFrameToRender = Frame;
 
-			UE_LOG(LogHMD, VeryVerbose, TEXT("StartGameFrame %u %u"), Frame->FrameNumber, Frame->ShowFlags.Rendering);
+			UE_LOG(LogHMD, VeryVerbose, TEXT("StartGameFrame %u"), Frame->FrameNumber);
+
+			if (!Splash->IsShown())
+			{
+				if (ovrp_GetInitialized() && WaitFrameNumber != Frame->FrameNumber)
+				{
+					UE_LOG(LogHMD, Verbose, TEXT("ovrp_WaitToBeginFrame %u"), Frame->FrameNumber);
+
+					ovrpResult Result;
+					if (OVRP_FAILURE(Result = ovrp_WaitToBeginFrame(Frame->FrameNumber)))
+					{
+						UE_LOG(LogHMD, Error, TEXT("ovrp_WaitToBeginFrame %u failed (%d)"), Frame->FrameNumber, Result);
+					}
+					else
+					{
+						WaitFrameNumber = Frame->FrameNumber;
+					}
+				}
+			}
 
 			UpdateStereoRenderingParams();
 		}
@@ -2962,27 +2995,11 @@ namespace OculusHMD
 			UE_LOG(LogHMD, VeryVerbose, TEXT("StartRenderFrame %u"), NextFrameToRender->FrameNumber);
 
 			LastFrameToRender = NextFrameToRender;
-			NextFrameToRender->Flags.bSplashIsShown = Splash->IsShown();// || NextFrameToRender->FrameNumber != NextFrameNumber;
-
-			if (GetSplash())
-			{
-				Splash->StopTicker();
-			}
+			NextFrameToRender->Flags.bSplashIsShown = Splash->IsShown();
 
 			if (NextFrameToRender->ShowFlags.Rendering && !NextFrameToRender->Flags.bSplashIsShown)
 			{
-				UE_LOG(LogHMD, Verbose, TEXT("ovrp_WaitToBeginFrame %u"), NextFrameToRender->FrameNumber);
-
-				ovrpResult Result;
-				if (OVRP_FAILURE(Result = ovrp_WaitToBeginFrame(NextFrameToRender->FrameNumber)))
-				{
-					UE_LOG(LogHMD, Error, TEXT("ovrp_WaitToBeginFrame %u failed (%d)"), NextFrameToRender->FrameNumber, Result);
-					NextFrameToRender->ShowFlags.Rendering = false;
-				}
-				else
-				{
-					NextFrameNumber++;
-				}
+				NextFrameNumber++;
 			}
 
 			FSettingsPtr XSettings = Settings->Clone();
diff --git a/Engine/Plugins/Runtime/Oculus/OculusVR/Source/OculusHMD/Private/OculusHMD.h b/Engine/Plugins/Runtime/Oculus/OculusVR/Source/OculusHMD/Private/OculusHMD.h
--- a/Engine/Plugins/Runtime/Oculus/OculusVR/Source/OculusHMD/Private/OculusHMD.h
+++ b/Engine/Plugins/Runtime/Oculus/OculusVR/Source/OculusHMD/Private/OculusHMD.h
@@ -225,6 +225,7 @@ protected:
 	void ApplicationResumeDelegate();
 	void SetupOcclusionMeshes();
 	void UpdateStereoRenderingParams();
+	void UpdateSplashScreen_GameThread();
 	void UpdateHmdRenderInfo();
 	void InitializeEyeLayer_RenderThread(FRHICommandListImmediate& RHICmdList);
 	void ApplySystemOverridesOnStereo(bool force = false);
@@ -320,6 +321,7 @@ public:
 	const int GetNextFrameNumber() const { return NextFrameNumber; }
 
 	const FRotator GetSplashRotation() const { return SplashRotation; }
+	void SetSplashRotationToForward();
 
 	void StartGameFrame_GameThread(); // Called from OnStartGameFrame
 	void FinishGameFrame_GameThread(); // Called from OnEndGameFrame
@@ -363,6 +365,7 @@ protected:
 
 			uint64	bNeedEnableStereo : 1;
 			uint64	bNeedDisableStereo : 1;
+			uint64  bNeedSplashUpdate : 1;
 		};
 		uint64 Raw;
 	} Flags;
@@ -403,6 +406,7 @@ protected:
 	// Game thread
 	FSettingsPtr Settings;
 	uint32 NextFrameNumber;
+	uint32 WaitFrameNumber;
 	FGameFramePtr Frame; // Valid from OnStartGameFrame to OnEndGameFrame
 	FGameFramePtr NextFrameToRender; // Valid from OnStartGameFrame to BeginRenderViewFamily
 	FGameFramePtr LastFrameToRender; // Valid from OnStartGameFrame to BeginRenderViewFamily
diff --git a/Engine/Plugins/Runtime/Oculus/OculusVR/Source/OculusHMD/Private/OculusHMD_Splash.cpp b/Engine/Plugins/Runtime/Oculus/OculusVR/Source/OculusHMD/Private/OculusHMD_Splash.cpp
--- a/Engine/Plugins/Runtime/Oculus/OculusVR/Source/OculusHMD/Private/OculusHMD_Splash.cpp
+++ b/Engine/Plugins/Runtime/Oculus/OculusVR/Source/OculusHMD/Private/OculusHMD_Splash.cpp
@@ -28,7 +28,6 @@ FSplash::FSplash(FOculusHMD* InOculusHMD) :
 	FramesOutstanding(0),
 	NextLayerId(1),
 	bInitialized(false),
-	bTickable(false),
 	bIsShown(false),
 	SystemDisplayInterval(1 / 90.0f)
 {
@@ -109,7 +108,17 @@ void FSplash::LoadSettings()
 	{
 		AddSplash(SplashDesc);
 	}
+
 	UStereoLayerFunctionLibrary::EnableAutoLoadingSplashScreen(HMDSettings->bAutoEnabled);
+	if (HMDSettings->bAutoEnabled)
+	{
+		FCoreUObjectDelegates::PreLoadMap.AddSP(this, &FSplash::OnPreLoadMap);
+	}
+}
+
+void FSplash::OnPreLoadMap(const FString&)
+{
+	Show();
 }
 
 void FSplash::Startup()
@@ -131,12 +140,7 @@ void FSplash::Startup()
 
 		LoadSettings();
 
-		Ticker = MakeShareable(new FTicker(this));
 
-		ExecuteOnRenderThread_DoNotWait([this]()
-		{
-			Ticker->Register();
-		});
 
 		bInitialized = true;
 	}
@@ -144,15 +148,38 @@ void FSplash::Startup()
 
 void FSplash::StopTicker()
 {
-	FScopeLock ScopeLock(&RenderThreadLock);
+	CheckInGameThread();
 
-	if (!IsShown())
+	if (!bIsShown)
 	{
-		bTickable = false;
+		ExecuteOnRenderThread([this]()
+		{
+			if (this->Ticker.IsValid())
+			{
+				this->Ticker->Unregister();
+				this->Ticker = nullptr;
+			}
+		});
 		UnloadTextures();
 	}
 }
 
+void FSplash::StartTicker()
+{
+	CheckInGameThread();
+
+	if (!Ticker.IsValid())
+	{
+		Ticker = MakeShareable(new FTicker(this));
+
+		ExecuteOnRenderThread([this]()
+		{
+			this->Ticker->Register();
+		});
+	}
+
+}
+
 void FSplash::RenderFrame_RenderThread(FRHICommandListImmediate& RHICmdList)
 {
 	CheckInRenderThread();
@@ -180,7 +207,8 @@ void FSplash::RenderFrame_RenderThread(FRHICommandListImmediate& RHICmdList)
 	}
 	else
 	{
-		OculusHMD->NextFrameNumber++;
+		OculusHMD->WaitFrameNumber = XFrame->FrameNumber;
+		OculusHMD->NextFrameNumber = XFrame->FrameNumber + 1;
 		FPlatformAtomics::InterlockedIncrement(&FramesOutstanding);
 	}
 
@@ -299,9 +327,6 @@ void FSplash::ReleaseResources_RHIThread()
 void FSplash::PreShutdown()
 {
 	CheckInGameThread();
-
-	// force Ticks to stop
-	bTickable = false;
 }
 
 
@@ -311,12 +336,13 @@ void FSplash::Shutdown()
 
 	if (bInitialized)
 	{
-		bTickable = false;
-
 		ExecuteOnRenderThread([this]()
 		{
-			Ticker->Unregister();
-			Ticker = nullptr;
+			if(Ticker)
+			{
+				Ticker->Unregister();
+				Ticker = nullptr;
+			}
 
 			ExecuteOnRHIThread([this]()
 			{
@@ -384,6 +410,8 @@ void FSplash::Show()
 
 	OculusHMD->InitDevice();
 
+	OculusHMD->SetSplashRotationToForward();
+
 	// Create new textures
 	UnloadTextures();
 
@@ -406,10 +434,7 @@ void FSplash::Show()
 		}
 	}
 
-	if (bWaitForRT)
-	{
-		FlushRenderingCommands();
-	}
+	FlushRenderingCommands();
 
 	for (int32 SplashLayerIndex = 0; SplashLayerIndex < SplashLayers.Num(); ++SplashLayerIndex)
 	{
@@ -453,7 +478,8 @@ void FSplash::Show()
 		FOculusSplashDesc UESplashDesc = OculusHMD->GetUESplashScreenDesc();
 		if (UESplashDesc.LoadedTexture != nullptr)
 		{
-			UELayer->SetDesc(StereoLayerDescFromOculusSplashDesc(UESplashDesc));
+			UELayer.Reset();
+			UELayer = MakeShareable(new FLayer(NextLayerId++, StereoLayerDescFromOculusSplashDesc(UESplashDesc)));
 			Layers_RenderThread_Input.Add(UELayer->Clone());
 		}
 
@@ -461,7 +487,7 @@ void FSplash::Show()
 	}
 
 	// If no textures are loaded, this will push black frame
-	bTickable = true;
+	StartTicker();
 	bIsShown = true;
 
 	UE_LOG(LogHMD, Log, TEXT("FSplash::OnShow"));
@@ -473,8 +499,9 @@ void FSplash::Hide()
 	CheckInGameThread();
 
 	UE_LOG(LogHMD, Log, TEXT("FSplash::OnHide"));
-//	bTickable = false;
 	bIsShown = false;
+
+	StopTicker();
 }
 
 void FSplash::UnloadTextures()
diff --git a/Engine/Plugins/Runtime/Oculus/OculusVR/Source/OculusHMD/Private/OculusHMD_Splash.h b/Engine/Plugins/Runtime/Oculus/OculusVR/Source/OculusHMD/Private/OculusHMD_Splash.h
--- a/Engine/Plugins/Runtime/Oculus/OculusVR/Source/OculusHMD/Private/OculusHMD_Splash.h
+++ b/Engine/Plugins/Runtime/Oculus/OculusVR/Source/OculusHMD/Private/OculusHMD_Splash.h
@@ -43,7 +43,8 @@ protected:
 
 		virtual void Tick(float DeltaTime) override { pSplash->Tick_RenderThread(DeltaTime); }
 		virtual TStatId GetStatId() const override  { RETURN_QUICK_DECLARE_CYCLE_STAT(FSplash, STATGROUP_Tickables); }
-		virtual bool IsTickable() const override	{ return pSplash->IsTickable(); }
+		virtual bool IsTickable() const override { return true; }
+	
 	protected:
 		FSplash* pSplash;
 	};
@@ -53,7 +54,6 @@ public:
 	virtual ~FSplash();
 
 	void Tick_RenderThread(float DeltaTime);
-	bool IsTickable() const { return bTickable; }
 	bool IsShown() const { return bIsShown; }
 
 	void Startup();
@@ -62,10 +62,13 @@ public:
 	void PreShutdown();
 	void Shutdown();
 
+	void OnPreLoadMap(const FString&);
+
 	int AddSplash(const FOculusSplashDesc&);
 	void ClearSplashes();
 	bool GetSplash(unsigned index, FOculusSplashDesc& OutDesc);
 	void StopTicker();
+	void StartTicker();
 
 	void Show();
 	void Hide();
@@ -96,7 +99,6 @@ protected:
 
 	// All these flags are only modified from the Game thread
 	bool bInitialized;
-	bool bTickable;
 	bool bIsShown;
 
 	float SystemDisplayInterval;
