Fixes Oculus splash screen related assert on shutdown.

diff --git a/Engine/Plugins/Runtime/Oculus/OculusVR/Source/OculusHMD/Private/OculusHMD_Splash.cpp b/Engine/Plugins/Runtime/Oculus/OculusVR/Source/OculusHMD/Private/OculusHMD_Splash.cpp
--- a/Engine/Plugins/Runtime/Oculus/OculusVR/Source/OculusHMD/Private/OculusHMD_Splash.cpp
+++ b/Engine/Plugins/Runtime/Oculus/OculusVR/Source/OculusHMD/Private/OculusHMD_Splash.cpp
@@ -348,6 +348,7 @@ void FSplash::Shutdown()
 			{
 				SplashLayers.Reset();
 				Layers_RenderThread.Reset();
+				Layers_RenderThread_Input.Reset();
 				Layers_RHIThread.Reset();
 			});
 		});
