Hotfix from Oculus: allocate the new eye buffers when we notice the need for them in UpdateStereoRenderingParams; fixes memory leak with adaptive resolution (varying pixel density).

diff --git a/Engine/Plugins/Runtime/Oculus/OculusVR/Source/OculusHMD/Private/OculusHMD.cpp b/Engine/Plugins/Runtime/Oculus/OculusVR/Source/OculusHMD/Private/OculusHMD.cpp
--- a/Engine/Plugins/Runtime/Oculus/OculusVR/Source/OculusHMD/Private/OculusHMD.cpp
+++ b/Engine/Plugins/Runtime/Oculus/OculusVR/Source/OculusHMD/Private/OculusHMD.cpp
@@ -1237,6 +1237,20 @@ namespace OculusHMD
 		check(InOutSizeX != 0 && InOutSizeY != 0);
 	}
 
+	void FOculusHMD::AllocateEyeBuffer()
+	{
+		CheckInGameThread();
+
+		ExecuteOnRenderThread([&]()
+		{
+			InitializeEyeLayer_RenderThread(GetImmediateCommandList_ForRenderCommand());
+
+			const FTextureSetProxyPtr& TextureSetProxy = EyeLayer_RenderThread->GetTextureSetProxy();
+			UE_LOG(LogHMD, Log, TEXT("Allocating Oculus %d x %d rendertarget swapchain"), TextureSetProxy->GetTexture2D()->GetSizeX(), TextureSetProxy->GetTexture2D()->GetSizeY());
+		});
+
+		bNeedReAllocateViewportRenderTarget = true;
+	}
 
 	bool FOculusHMD::NeedReAllocateViewportRenderTarget(const FViewport& Viewport)
 	{
@@ -1266,12 +1280,7 @@ namespace OculusHMD
 
 		if (LayerMap[0].IsValid())
 		{
-			InitializeEyeLayer_RenderThread(GetImmediateCommandList_ForRenderCommand());
-
-			UE_LOG(LogHMD, Log, TEXT("Allocating Oculus %d x %d rendertarget swapchain"), SizeX, SizeY);
-
 			const FTextureSetProxyPtr& TextureSetProxy = EyeLayer_RenderThread->GetTextureSetProxy();
-
 			if (TextureSetProxy.IsValid())
 			{
 				OutTargetableTexture = TextureSetProxy->GetTexture2D();
@@ -2332,7 +2341,7 @@ namespace OculusHMD
 			// Flag if need to recreate render targets
 			if (!EyeLayer->CanReuseResources(EyeLayer_RenderThread.Get()))
 			{
-				bNeedReAllocateViewportRenderTarget = true;
+				AllocateEyeBuffer();
 			}
 		}
 	}
@@ -2347,6 +2356,7 @@ namespace OculusHMD
 
 	void FOculusHMD::InitializeEyeLayer_RenderThread(FRHICommandListImmediate& RHICmdList)
 	{
+		check(!InGameThread());
 		CheckInRenderThread();
 
 		if (LayerMap[0].IsValid())
diff --git a/Engine/Plugins/Runtime/Oculus/OculusVR/Source/OculusHMD/Private/OculusHMD.h b/Engine/Plugins/Runtime/Oculus/OculusVR/Source/OculusHMD/Private/OculusHMD.h
--- a/Engine/Plugins/Runtime/Oculus/OculusVR/Source/OculusHMD/Private/OculusHMD.h
+++ b/Engine/Plugins/Runtime/Oculus/OculusVR/Source/OculusHMD/Private/OculusHMD.h
@@ -186,6 +186,7 @@ public:
 	virtual bool AllocateDepthTexture(uint32 Index, uint32 SizeX, uint32 SizeY, uint8 Format, uint32 NumMips, uint32 InTexFlags, uint32 TargetableTextureFlags, FTexture2DRHIRef& OutTargetableTexture, FTexture2DRHIRef& OutShaderResourceTexture, uint32 NumSamples = 1) override;
 	virtual void UpdateViewportWidget(bool bUseSeparateRenderTarget, const class FViewport& Viewport, class SViewport* ViewportWidget) override;
 	virtual FXRRenderBridge* GetActiveRenderBridge_GameThread(bool bUseSeparateRenderTarget);
+	void AllocateEyeBuffer();
 
 	// IStereoLayers interface
 	virtual uint32 CreateLayer(const IStereoLayers::FLayerDesc& InLayerDesc) override;
