Don't wait on the same frame multiple times in OculusHMD_Splash.
Fixes remaining log spam in #jira UE-69042
Cherrypicked CL 7039441 from Dev-VR

diff --git a/Engine/Plugins/Runtime/Oculus/OculusVR/Source/OculusHMD/Private/OculusHMD_Splash.cpp b/Engine/Plugins/Runtime/Oculus/OculusVR/Source/OculusHMD/Private/OculusHMD_Splash.cpp
--- a/Engine/Plugins/Runtime/Oculus/OculusVR/Source/OculusHMD/Private/OculusHMD_Splash.cpp
+++ b/Engine/Plugins/Runtime/Oculus/OculusVR/Source/OculusHMD/Private/OculusHMD_Splash.cpp
@@ -207,7 +207,7 @@ void FSplash::RenderFrame_RenderThread(FRHICommandListImmediate& RHICmdList)
 	}
 
 	ovrpResult Result;
-	if ( ovrp_GetInitialized() )
+	if ( ovrp_GetInitialized() && OculusHMD->WaitFrameNumber != Frame->FrameNumber)
 	{ 
 		UE_LOG(LogHMD, Verbose, TEXT("Splash ovrp_WaitToBeginFrame %u"), XFrame->FrameNumber);
 		if (OVRP_FAILURE(Result = ovrp_WaitToBeginFrame(XFrame->FrameNumber)))
