Fix issue with helper function on unsupported Occulus platforms

#ue4
#android

diff --git a/Engine/Plugins/Runtime/Oculus/OculusVR/Source/OculusHMD/Private/OculusFunctionLibrary.cpp b/Engine/Plugins/Runtime/Oculus/OculusVR/Source/OculusHMD/Private/OculusFunctionLibrary.cpp
--- a/Engine/Plugins/Runtime/Oculus/OculusVR/Source/OculusHMD/Private/OculusFunctionLibrary.cpp
+++ b/Engine/Plugins/Runtime/Oculus/OculusVR/Source/OculusHMD/Private/OculusFunctionLibrary.cpp
@@ -605,6 +605,7 @@ class IStereoLayers* UOculusFunctionLibrary::GetStereoLayers()
 }
 
 /** Helper that converts EBoundaryType to ovrpBoundaryType */
+#if OCULUS_HMD_SUPPORTED_PLATFORMS
 static ovrpBoundaryType ToOvrpBoundaryType(EBoundaryType Source)
 {
 	switch (Source)
@@ -617,6 +618,7 @@ static ovrpBoundaryType ToOvrpBoundaryType(EBoundaryType Source)
 		return ovrpBoundary_Outer;
 	}
 }
+#endif // OCULUS_HMD_SUPPORTED_PLATFORMS
 
 bool UOculusFunctionLibrary::IsGuardianDisplayed()
 {
