Oculus haptic fix.

diff --git a/Engine/Plugins/Runtime/Oculus/OculusVR/Source/OculusInput/Private/OculusInput.cpp b/Engine/Plugins/Runtime/Oculus/OculusVR/Source/OculusInput/Private/OculusInput.cpp
--- a/Engine/Plugins/Runtime/Oculus/OculusVR/Source/OculusInput/Private/OculusInput.cpp
+++ b/Engine/Plugins/Runtime/Oculus/OculusVR/Source/OculusInput/Private/OculusInput.cpp
@@ -624,11 +624,11 @@ void FOculusInput::SetChannelValue( int32 ControllerId, FForceFeedbackChannelTyp
 			// @todo: The SMALL channel controls frequency, the LARGE channel controls amplitude.  This is a bit of a weird fit.
 			if( ChannelType == FForceFeedbackChannelType::LEFT_SMALL || ChannelType == FForceFeedbackChannelType::RIGHT_SMALL )
 			{
-				ControllerState.HapticFrequency = Value;
+				ControllerState.ForceFeedbackHapticFrequency = Value;
 			}
 			else
 			{
-				ControllerState.HapticAmplitude = Value;
+				ControllerState.ForceFeedbackHapticAmplitude = Value;
 			}
 
 			UpdateForceFeedback( ControllerPair, Hand );
@@ -648,16 +648,16 @@ void FOculusInput::SetChannelValues( int32 ControllerId, const FForceFeedbackVal
 			FOculusTouchControllerState& LeftControllerState = ControllerPair.ControllerStates[ (int32)EControllerHand::Left ];
 			if (!LeftControllerState.bPlayingHapticEffect)
 			{
-				LeftControllerState.HapticFrequency = Values.LeftSmall;
-				LeftControllerState.HapticAmplitude = Values.LeftLarge;
+				LeftControllerState.ForceFeedbackHapticFrequency = Values.LeftSmall;
+				LeftControllerState.ForceFeedbackHapticAmplitude = Values.LeftLarge;
 				UpdateForceFeedback(ControllerPair, EControllerHand::Left);
 			}
 
 			FOculusTouchControllerState& RightControllerState = ControllerPair.ControllerStates[(int32)EControllerHand::Right];
 			if (!RightControllerState.bPlayingHapticEffect)
 			{
-				RightControllerState.HapticFrequency = Values.RightSmall;
-				RightControllerState.HapticAmplitude = Values.RightLarge;
+				RightControllerState.ForceFeedbackHapticFrequency = Values.RightSmall;
+				RightControllerState.ForceFeedbackHapticAmplitude = Values.RightLarge;
 				UpdateForceFeedback(ControllerPair, EControllerHand::Right);
 			}
 		}
@@ -681,10 +681,10 @@ void FOculusInput::UpdateForceFeedback( const FOculusTouchControllerPair& Contro
 				GetHapticFrequencyRange(FreqMin, FreqMax);
 
 				// Map the [0.0 - 1.0] range to a useful range of frequencies for the Oculus controllers
-				const float ActualFrequency = FMath::Lerp(FreqMin, FreqMax, FMath::Clamp(ControllerState.HapticFrequency, 0.0f, 1.0f));
+				const float ActualFrequency = FMath::Lerp(FreqMin, FreqMax, FMath::Clamp(ControllerState.ForceFeedbackHapticFrequency, 0.0f, 1.0f));
 
 				// Oculus SDK wants amplitude values between 0.0 and 1.0
-				const float ActualAmplitude = ControllerState.HapticAmplitude * GetHapticAmplitudeScale();
+				const float ActualAmplitude = ControllerState.ForceFeedbackHapticAmplitude * GetHapticAmplitudeScale();
 
 				ovrpController OvrController = ovrpController_None;
 				if (OvrpControllerState.ConnectedControllerTypes & (ovrpController_Touch))
@@ -847,7 +847,13 @@ void FOculusInput::SetHapticFeedbackValues(int32 ControllerId, int32 Hand, const
 
 					ovrpControllerState4 OvrpControllerState;
 					
-					if (OVRP_SUCCESS(ovrp_GetControllerState4((ovrpController)(ovrpController_Active | ovrpController_LTrackedRemote | ovrpController_RTrackedRemote), &OvrpControllerState)) &&
+					ovrpController ControllerTypes = (ovrpController)(ovrpController_Active | ovrpController_LTrackedRemote | ovrpController_RTrackedRemote);
+
+#ifdef USE_ANDROID_INPUT
+					ControllerTypes = (ovrpController)(controllerTypes | ovrpController_Touch);
+#endif
+
+					if (OVRP_SUCCESS(ovrp_GetControllerState4(ControllerTypes, &OvrpControllerState)) &&
 						(OvrpControllerState.ConnectedControllerTypes & (ovrpController_Touch | ovrpController_LTrackedRemote | ovrpController_RTrackedRemote)))
 					{
 						// Buffered haptics is currently only supported on Touch
@@ -938,7 +944,9 @@ void FOculusInput::SetHapticFeedbackValues(int32 ControllerId, int32 Hand, const
 							float FreqMin, FreqMax = 0.f;
 							GetHapticFrequencyRange(FreqMin, FreqMax);
 
-							const float Frequency = FMath::Lerp(FreqMin, FreqMax, FMath::Clamp(Values.Frequency, 0.f, 1.f));
+							const float InitialFreq = (Values.Frequency > 0.0f) ? Values.Frequency : 1.0f;
+							const float Frequency = FMath::Lerp(FreqMin, FreqMax, FMath::Clamp(InitialFreq, 0.f, 1.f));
+
 							const float Amplitude = Values.Amplitude * GetHapticAmplitudeScale();
 
 							if (ControllerState.HapticAmplitude != Amplitude || ControllerState.HapticFrequency != Frequency)
diff --git a/Engine/Plugins/Runtime/Oculus/OculusVR/Source/OculusInput/Private/OculusInputState.h b/Engine/Plugins/Runtime/Oculus/OculusVR/Source/OculusInput/Private/OculusInputState.h
--- a/Engine/Plugins/Runtime/Oculus/OculusVR/Source/OculusInput/Private/OculusInputState.h
+++ b/Engine/Plugins/Runtime/Oculus/OculusVR/Source/OculusInput/Private/OculusInputState.h
@@ -247,6 +247,12 @@ struct FOculusTouchControllerState
 	/** Haptic amplitude (zero to disable) */
 	float HapticAmplitude;
 
+	/** Force feedback haptic frequency (zero to disable) */
+	float ForceFeedbackHapticFrequency;
+
+	/** Force feedback haptic amplitude (zero to disable) */
+	float ForceFeedbackHapticAmplitude;
+
 	/** Number of times that controller was recentered (for mobile controllers) */
 	int RecenterCount;
 
@@ -262,6 +268,8 @@ struct FOculusTouchControllerState
 		  bPlayingHapticEffect( false ),
 		  HapticFrequency( 0.0f ),
 		  HapticAmplitude( 0.0f ),
+		  ForceFeedbackHapticFrequency(0.0f),
+		  ForceFeedbackHapticAmplitude(0.0f),
 		  RecenterCount(0)
 	{
 		for( FOculusButtonState& Button : Buttons )
