Fix incorrect variable name, causing compile errors on Android.

diff --git a/Engine/Plugins/Runtime/Oculus/OculusVR/Source/OculusInput/Private/OculusInput.cpp b/Engine/Plugins/Runtime/Oculus/OculusVR/Source/OculusInput/Private/OculusInput.cpp
--- a/Engine/Plugins/Runtime/Oculus/OculusVR/Source/OculusInput/Private/OculusInput.cpp
+++ b/Engine/Plugins/Runtime/Oculus/OculusVR/Source/OculusInput/Private/OculusInput.cpp
@@ -850,7 +850,7 @@ void FOculusInput::SetHapticFeedbackValues(int32 ControllerId, int32 Hand, const
 					ovrpController ControllerTypes = (ovrpController)(ovrpController_Active | ovrpController_LTrackedRemote | ovrpController_RTrackedRemote);
 
 #ifdef USE_ANDROID_INPUT
-					ControllerTypes = (ovrpController)(controllerTypes | ovrpController_Touch);
+					ControllerTypes = (ovrpController)(ControllerTypes | ovrpController_Touch);
 #endif
 
 					if (OVRP_SUCCESS(ovrp_GetControllerState4(ControllerTypes, &OvrpControllerState)) &&
