UE-65161 Left side of spectator screen renders completely white when mode set to Undistorted
UE-65162 Spectator screen does not render camera view when mode is set to TexturePlusEye
-Need to limit the alpha clear to the half of the spectator destination texture that we are actually going to write to.

diff --git a/Engine/Plugins/Runtime/Oculus/OculusVR/Source/OculusHMD/Private/OculusHMD_CustomPresent.cpp b/Engine/Plugins/Runtime/Oculus/OculusVR/Source/OculusHMD/Private/OculusHMD_CustomPresent.cpp
--- a/Engine/Plugins/Runtime/Oculus/OculusVR/Source/OculusHMD/Private/OculusHMD_CustomPresent.cpp
+++ b/Engine/Plugins/Runtime/Oculus/OculusVR/Source/OculusHMD/Private/OculusHMD_CustomPresent.cpp
@@ -401,6 +401,7 @@ void FCustomPresent::CopyTexture_RenderThread(FRHICommandListImmediate& RHICmdLi
 
 		if (bNoAlphaWrite)
 		{
+			RHICmdList.SetViewport(DstRect.Min.X, DstRect.Min.Y, 0.0f, DstRect.Max.X, DstRect.Max.Y, 1.0f);
 			DrawClearQuad(RHICmdList, bAlphaPremultiply ? FLinearColor::Black : FLinearColor::White);
 		}
 
