UE-65927 Oculus VR: GetPlayAreaTransform rotation sometimes reversed
-If the tracking to world transform was in the 'back' 180 degrees the play area would still be oriented toward the 'front' 180 degrees.
-The fix comes from oculus, I am the reviewer and tested it as well.

diff --git a/Engine/Plugins/Runtime/Oculus/OculusVR/Source/OculusHMD/Private/OculusFunctionLibrary.cpp b/Engine/Plugins/Runtime/Oculus/OculusVR/Source/OculusHMD/Private/OculusFunctionLibrary.cpp
--- a/Engine/Plugins/Runtime/Oculus/OculusVR/Source/OculusHMD/Private/OculusFunctionLibrary.cpp
+++ b/Engine/Plugins/Runtime/Oculus/OculusVR/Source/OculusHMD/Private/OculusFunctionLibrary.cpp
@@ -708,7 +708,10 @@ FTransform UOculusFunctionLibrary::GetPlayAreaTransform()
 
 			float metersScale = OculusHMD->GetWorldToMetersScale();
 
-			FQuat Rotation(FVector::UpVector, FMath::Acos((ConvertedPoints[1] - ConvertedPoints[0]).GetSafeNormal() | FVector::RightVector));
+			FVector Edge = ConvertedPoints[1] - ConvertedPoints[0];
+			float Angle = FMath::Acos((Edge).GetSafeNormal() | FVector::RightVector);
+			FQuat Rotation(FVector::UpVector, Edge.X < 0 ? Angle : -Angle);
+			
 			FVector Position = (ConvertedPoints[0] + ConvertedPoints[1] + ConvertedPoints[2] + ConvertedPoints[3]) / 4;
 			FVector Scale(FVector::Distance(ConvertedPoints[3], ConvertedPoints[0]) / metersScale, FVector::Distance(ConvertedPoints[1], ConvertedPoints[0]) / metersScale, 1.0);
 
