UE-66044 Oculus VR toggle for editor-only
-Fix from oculus: ovrpInitializeFlag_SupportsVRToggle is only desired in the editor, it had been on always.

diff --git a/Engine/Plugins/Runtime/Oculus/OculusVR/Source/OculusHMD/Private/OculusHMD.cpp b/Engine/Plugins/Runtime/Oculus/OculusVR/Source/OculusHMD/Private/OculusHMD.cpp
--- a/Engine/Plugins/Runtime/Oculus/OculusVR/Source/OculusHMD/Private/OculusHMD.cpp
+++ b/Engine/Plugins/Runtime/Oculus/OculusVR/Source/OculusHMD/Private/OculusHMD.cpp
@@ -1887,7 +1887,7 @@ namespace OculusHMD
 			void* activity = nullptr;
 #endif
 
-			int initializeFlags = ovrpInitializeFlag_SupportsVRToggle;
+			int initializeFlags = GIsEditor ? ovrpInitializeFlag_SupportsVRToggle : 0;
 
 			if (Settings->Flags.bSupportsDash)
 			{
