Restored some functions that had been deleted from the Oculus Function Library and deprecated them instead to provide proper information about how to update any code/blueprints using them.

diff --git a/Engine/Plugins/Runtime/Oculus/OculusVR/Source/OculusHMD/Private/OculusFunctionLibrary.cpp b/Engine/Plugins/Runtime/Oculus/OculusVR/Source/OculusHMD/Private/OculusFunctionLibrary.cpp
--- a/Engine/Plugins/Runtime/Oculus/OculusVR/Source/OculusHMD/Private/OculusFunctionLibrary.cpp
+++ b/Engine/Plugins/Runtime/Oculus/OculusVR/Source/OculusHMD/Private/OculusFunctionLibrary.cpp
@@ -3,6 +3,9 @@
 #include "OculusFunctionLibrary.h"
 #include "OculusHMDPrivate.h"
 #include "OculusHMD.h"
+#include "Logging/MessageLog.h"
+
+#define LOCTEXT_NAMESPACE "OculusFunctionLibrary"
 
 //-------------------------------------------------------------------------------------------------
 // UOculusFunctionLibrary
@@ -237,6 +240,82 @@ void UOculusFunctionLibrary::ClearLoadingSplashScreens()
 #endif // OCULUS_HMD_SUPPORTED_PLATFORMS
 }
 
+void UOculusFunctionLibrary::ShowLoadingSplashScreen()
+{
+	FText Message = LOCTEXT("ShowLoadingSplashScreen",
+		"UOculusFunctionLibrary::ShowLoadingSplashScreen has been deprecated and no longer functions as before.  Use the generic UStereoLayerFunctionLibrary::ShowSplashScreen instead");
+	UE_LOG(LogHMD, Error, TEXT("%s"), *(Message.ToString()));
+	FMessageLog("PIE").Error(Message);
+}
+
+void UOculusFunctionLibrary::HideLoadingSplashScreen(bool bClear)
+{
+	FText Message = LOCTEXT("HideLoadingSplashScreenDeprecated",
+		"UOculusFunctionLibrary::HideLoadingSplashScreen has been deprecated and no longer functions as before.  Use the generic UStereoLayerFunctionLibrary::HideSplashScreen instead");
+	UE_LOG(LogHMD, Error, TEXT("%s"), *(Message.ToString()));
+	FMessageLog("PIE").Error(Message);
+}
+
+void UOculusFunctionLibrary::EnableAutoLoadingSplashScreen(bool bAutoShowEnabled)
+{
+	FText Message = LOCTEXT("EnableAutoLoadingSplashScreenDeprecated",
+		"UOculusFunctionLibrary::EnableAutoLoadingSplashScreen has been deprecated and no longer functions as before.  Use the generic UStereoLayerFunctionLibrary::EnableAutoLoadingSplashScreen instead");
+	UE_LOG(LogHMD, Error, TEXT("%s"), *(Message.ToString()));
+	FMessageLog("PIE").Error(Message);
+}
+
+bool UOculusFunctionLibrary::IsAutoLoadingSplashScreenEnabled()
+{
+	FText Message = LOCTEXT("IsAutoLoadingSplashScreenEnabledDeprecated",
+		"UOculusFunctionLibrary::IsAutoLoadingSplashScreenEnabled has been deprecated and no longer functions as before. Please use the generic UStereoLayerFunctionLibrary instead.");
+	UE_LOG(LogHMD, Error, TEXT("%s"), *(Message.ToString()));
+	FMessageLog("PIE").Error(Message);
+
+	return false;
+}
+
+void UOculusFunctionLibrary::ShowLoadingIcon(class UTexture2D* Texture)
+{
+	FText Message = LOCTEXT("IsAutoLoadingSplashScreenEnabledDeprecated",
+		"UOculusFunctionLibrary::ShowLoadingIcon has been deprecated and no longer functions as before. Please use the generic UStereoLayerFunctionLibrary instead.");
+	UE_LOG(LogHMD, Error, TEXT("%s"), *(Message.ToString()));
+	FMessageLog("PIE").Error(Message);
+}
+
+void UOculusFunctionLibrary::HideLoadingIcon()
+{
+	FText Message = LOCTEXT("HideLoadingIconDeprecated",
+		"UOculusFunctionLibrary::HideLoadingIcon has been deprecated and no longer functions as before. Please use the generic UStereoLayerFunctionLibrary instead.");
+	UE_LOG(LogHMD, Error, TEXT("%s"), *(Message.ToString()));
+	FMessageLog("PIE").Error(Message);
+}
+
+bool UOculusFunctionLibrary::IsLoadingIconEnabled()
+{
+	FText Message = LOCTEXT("IsLoadingIconEnabledDeprecated",
+		"UOculusFunctionLibrary::IsLoadingIconEnabled has been deprecated and no longer functions as before. Please use the generic UStereoLayerFunctionLibrary instead.");
+	UE_LOG(LogHMD, Error, TEXT("%s"), *(Message.ToString()));
+	FMessageLog("PIE").Error(Message);
+	return false;
+}
+
+
+void UOculusFunctionLibrary::SetLoadingSplashParams(FString TexturePath, FVector DistanceInMeters, FVector2D SizeInMeters, FVector RotationAxis, float RotationDeltaInDeg)
+{
+	FText Message = LOCTEXT("SetLoadingSplashParamsDeprecated",
+		"UOculusFunctionLibrary::SetLoadingSplashParams has been deprecated and no longer functions as before. Please use the generic UStereoLayerFunctionLibrary instead.");
+	UE_LOG(LogHMD, Error, TEXT("%s"), *(Message.ToString()));
+	FMessageLog("PIE").Error(Message);
+}
+
+void UOculusFunctionLibrary::GetLoadingSplashParams(FString& TexturePath, FVector& DistanceInMeters, FVector2D& SizeInMeters, FVector& RotationAxis, float& RotationDeltaInDeg)
+{
+	FText Message = LOCTEXT("GetLoadingSplashParamsDeprecated",
+		"UOculusFunctionLibrary::GetLoadingSplashParams has been deprecated and no longer functions as before. Please use the generic UStereoLayerFunctionLibrary instead.");
+	UE_LOG(LogHMD, Error, TEXT("%s"), *(Message.ToString()));
+	FMessageLog("PIE").Error(Message);
+}
+
 bool UOculusFunctionLibrary::HasInputFocus()
 {
 #if OCULUS_HMD_SUPPORTED_PLATFORMS
@@ -629,3 +708,5 @@ void UOculusFunctionLibrary::SetGuardianVisibility(bool GuardianVisible)
 	}
 #endif
 }
+
+#undef LOCTEXT_NAMESPACE
\ No newline at end of file
diff --git a/Engine/Plugins/Runtime/Oculus/OculusVR/Source/OculusHMD/Public/OculusFunctionLibrary.h b/Engine/Plugins/Runtime/Oculus/OculusVR/Source/OculusHMD/Public/OculusFunctionLibrary.h
--- a/Engine/Plugins/Runtime/Oculus/OculusVR/Source/OculusHMD/Public/OculusFunctionLibrary.h
+++ b/Engine/Plugins/Runtime/Oculus/OculusVR/Source/OculusHMD/Public/OculusFunctionLibrary.h
@@ -264,6 +264,85 @@ class OCULUSHMD_API UOculusFunctionLibrary : public UBlueprintFunctionLibrary
 	UFUNCTION(BlueprintCallable, Category = "OculusLibrary")
 	static void ClearLoadingSplashScreens();
 
+	/**
+	 * Shows loading splash screen.
+	 */
+	UE_DEPRECATED(4.22, "UOculusFunctionLibrary::ShowLoadingSplashScreen has been deprecated and no longer functions as before.  Use the generic UStereoLayerFunctionLibrary::ShowSplashScreen instead")
+	UFUNCTION(BlueprintCallable, Category = "OculusLibrary", meta = (DeprecatedFunction, DeprecationMessage = "UOculusFunctionLibrary::ShowLoadingSplashScreen has been deprecated and no longer functions as before.  Use the generic UStereoLayerFunctionLibrary::ShowSplashScreen instead"))
+	static void ShowLoadingSplashScreen();
+
+	/**
+	 * Hides loading splash screen.
+	 *
+	 * @param	bClear	(in) Clear all splash screens after hide.
+	 */
+	UE_DEPRECATED(4.22, "UOculusFunctionLibrary::HideLoadingSplashScreen has been deprecated and no longer functions as before.  Use the generic UStereoLayerFunctionLibrary::HideSplashScreen instead")
+	UFUNCTION(BlueprintCallable, Category = "OculusLibrary", meta = (DeprecatedFunction, DeprecationMessage = "UOculusFunctionLibrary::HideLoadingSplashScreen has been deprecated and no longer functions as before.  Use the generic UStereoLayerFunctionLibrary::HideSplashScreen instead"))
+	static void HideLoadingSplashScreen(bool bClear = false);
+
+	/**
+	 * Enables/disables splash screen to be automatically shown when LoadMap is called.
+	 *
+	 * @param	bAutoShowEnabled	(in)	True, if automatic showing of splash screens is enabled when map is being loaded.
+	 */
+	UE_DEPRECATED(4.22, "UOculusFunctionLibrary::EnableAutoLoadingSplashScreen has been deprecated and no longer functions as before.  Use the generic UStereoLayerFunctionLibrary::EnableAutoLoadingSplashScreen instead")
+	UFUNCTION(BlueprintCallable, Category = "OculusLibrary", meta = (DeprecatedFunction, DeprecationMessage = "UOculusFunctionLibrary::EnableAutoLoadingSplashScreen has been deprecated and no longer functions as before.  Use the generic UStereoLayerFunctionLibrary::EnableAutoLoadingSplashScreen instead"))
+	static void EnableAutoLoadingSplashScreen(bool bAutoShowEnabled);
+
+	/**
+	 * Returns true, if the splash screen is automatically shown when LoadMap is called.
+	 */
+	UE_DEPRECATED(4.22, "UOculusFunctionLibrary::IsAutoLoadingSplashScreenEnabled has been deprecated and no longer functions as before. Please use the generic UStereoLayerFunctionLibrary instead.")
+	UFUNCTION(BlueprintPure, Category = "OculusLibrary", meta = (DeprecatedFunction, DeprecationMessage = "UOculusFunctionLibrary::IsAutoLoadingSplashScreenEnabled has been deprecated and no longer functions as before. Please use the generic UStereoLayerFunctionLibrary instead."))
+	static bool IsAutoLoadingSplashScreenEnabled();
+
+	/**
+	 * Sets a texture for loading icon mode and shows it. This call will clear all the splashes.
+	 */
+	UE_DEPRECATED(4.22, "UOculusFunctionLibrary::ShowLoadingIcon has been deprecated and no longer functions as before. Please use the generic UStereoLayerFunctionLibrary instead.")
+	UFUNCTION(BlueprintCallable, Category = "OculusLibrary", meta = (DeprecatedFunction, DeprecationMessage = "UOculusFunctionLibrary::ShowLoadingIcon has been deprecated and no longer functions as before. Please use the generic UStereoLayerFunctionLibrary instead."))
+	static void ShowLoadingIcon(class UTexture2D* Texture);
+
+	/**
+	 * Clears the loading icon. This call will clear all the splashes.
+	 */
+	UE_DEPRECATED(4.22, "UOculusFunctionLibrary::HideLoadingIcon has been deprecated and no longer functions as before. Please use the generic UStereoLayerFunctionLibrary instead.")
+	UFUNCTION(BlueprintCallable, Category = "OculusLibrary", meta = (DeprecatedFunction, DeprecationMessage = "UOculusFunctionLibrary::HideLoadingIcon has been deprecated and no longer functions as before. Please use the generic UStereoLayerFunctionLibrary instead."))
+	static void HideLoadingIcon();
+
+	/**
+	 * Returns true, if the splash screen is in loading icon mode.
+	 */
+	UE_DEPRECATED(4.22, "UOculusFunctionLibrary::IsLoadingIconEnabled has been deprecated and no longer functions as before. Please use the generic UStereoLayerFunctionLibrary instead.")
+	UFUNCTION(BlueprintPure, Category = "OculusLibrary", meta = (DeprecatedFunction, DeprecationMessage = "UOculusFunctionLibrary::IsLoadingIconEnabled has been deprecated and no longer functions as before. Please use the generic UStereoLayerFunctionLibrary instead."))
+	static bool IsLoadingIconEnabled();
+
+	/**
+	 * Sets loading splash screen parameters.
+	 *
+	 * @param TexturePath		(in) A path to the texture asset to be used for the splash. Gear VR uses it as a path for loading icon; all other params are currently ignored by Gear VR.
+	 * @param DistanceInMeters	(in) Distance, in meters, to the center of the splash screen.
+	 * @param SizeInMeters		(in) Size, in meters, of the quad with the splash screen.
+	 * @param RotationAxes		(in) A vector that specifies the axis of the splash screen rotation (if RotationDelta is specified).
+	 * @param RotationDeltaInDeg (in) Rotation delta, in degrees, that is added each 2nd frame to the quad transform. The quad is rotated around the vector "RotationAxes".
+	 */
+	UE_DEPRECATED(4.22, "UOculusFunctionLibrary::SetLoadingSplashParams has been deprecated and no longer functions as before. Please use the generic UStereoLayerFunctionLibrary instead.")
+	UFUNCTION(BlueprintCallable, Category = "OculusLibrary", meta = (DeprecatedFunction, DeprecationMessage = "UOculusFunctionLibrary::SetLoadingSplashParams has been deprecated and no longer functions as before. Please use the generic UStereoLayerFunctionLibrary instead."))
+	static void SetLoadingSplashParams(FString TexturePath, FVector DistanceInMeters, FVector2D SizeInMeters, FVector RotationAxis, float RotationDeltaInDeg);
+
+	/**
+	 * Returns loading splash screen parameters.
+	 *
+	 * @param TexturePath		(out) A path to the texture asset to be used for the splash. Gear VR uses it as a path for loading icon; all other params are currently ignored by Gear VR.
+	 * @param DistanceInMeters	(out) Distance, in meters, to the center of the splash screen.
+	 * @param SizeInMeters		(out) Size, in meters, of the quad with the splash screen.
+	 * @param RotationAxes		(out) A vector that specifies the axis of the splash screen rotation (if RotationDelta is specified).
+	 * @param RotationDeltaInDeg (out) Rotation delta, in degrees, that is added each 2nd frame to the quad transform. The quad is rotated around the vector "RotationAxes".
+	 */
+	UE_DEPRECATED(4.22, "UOculusFunctionLibrary::GetLoadingSplashParams has been deprecated and no longer functions as before. Please use the generic UStereoLayerFunctionLibrary instead.")
+	UFUNCTION(BlueprintPure, Category = "OculusLibrary", meta = (DeprecatedFunction, DeprecationMessage = "UOculusFunctionLibrary::GetLoadingSplashParams has been deprecated and no longer functions as before. Please use the generic UStereoLayerFunctionLibrary instead."))
+	static void GetLoadingSplashParams(FString& TexturePath, FVector& DistanceInMeters, FVector2D& SizeInMeters, FVector& RotationAxis, float& RotationDeltaInDeg);
+
 	/**
 	* Returns true, if the app has input focus.
 	*/
