Oculus bug fixes for splash screen and recentering

diff --git a/Engine/Plugins/Runtime/Oculus/OculusVR/Source/OculusHMD/Private/OculusHMD.cpp b/Engine/Plugins/Runtime/Oculus/OculusVR/Source/OculusHMD/Private/OculusHMD.cpp
--- a/Engine/Plugins/Runtime/Oculus/OculusVR/Source/OculusHMD/Private/OculusHMD.cpp
+++ b/Engine/Plugins/Runtime/Oculus/OculusVR/Source/OculusHMD/Private/OculusHMD.cpp
@@ -383,51 +383,44 @@ namespace OculusHMD
 
 	void FOculusHMD::ResetOrientationAndPosition(float yaw)
 	{
-		CheckInGameThread();
-
-		if (NextFrameToRender)
-		{
-			const bool floorLevel = GetTrackingOrigin() == EHMDTrackingOrigin::Floor;
-			ovrpPoseStatef poseState;
-			ovrp_GetNodePoseState3(ovrpStep_Render, NextFrameToRender->FrameNumber, ovrpNode_Head, &poseState);
-			Settings->BaseOffset = ToFVector(poseState.Pose.Position);
-			if (floorLevel)
-				Settings->BaseOffset.Z = 0;
-
-			Settings->BaseOrientation = FRotator(0, FRotator(ToFQuat(poseState.Pose.Orientation)).Yaw - yaw, 0).Quaternion();
-		}
+		Recenter(RecenterOrientationAndPosition, yaw);
 	}
 
-
 	void FOculusHMD::ResetOrientation(float yaw)
 	{
-		CheckInGameThread();
-
-		if (NextFrameToRender)
-		{
-			ovrpPoseStatef poseState;
-			ovrp_GetNodePoseState3(ovrpStep_Render, NextFrameToRender->FrameNumber, ovrpNode_Head, &poseState);
-			Settings->BaseOrientation = FRotator(0, FRotator(ToFQuat(poseState.Pose.Orientation)).Yaw - yaw, 0).Quaternion();
-		}
+		Recenter(RecenterOrientation, yaw);
 	}
 
-
 	void FOculusHMD::ResetPosition()
+	{
+		Recenter(RecenterPosition, 0);
+	}
+
+	void FOculusHMD::Recenter(FRecenterTypes RecenterType, float Yaw)
 	{
 		CheckInGameThread();
 
 		if (NextFrameToRender)
 		{
-			const bool floorLevel = GetTrackingOrigin() == EHMDTrackingOrigin::Floor;
+			const bool floorLevel = GetTrackingOrigin() != EHMDTrackingOrigin::Eye;
 			ovrpPoseStatef poseState;
+			ovrp_Update3(ovrpStep_Render, NextFrameToRender->FrameNumber, 0.0);
 			ovrp_GetNodePoseState3(ovrpStep_Render, NextFrameToRender->FrameNumber, ovrpNode_Head, &poseState);
-			Settings->BaseOffset = ToFVector(poseState.Pose.Position);
-			if (floorLevel)
-				Settings->BaseOffset.Z = 0;
+
+			if (RecenterType & RecenterPosition)
+			{
+				Settings->BaseOffset = ToFVector(poseState.Pose.Position);
+				if (floorLevel)
+					Settings->BaseOffset.Z = 0;
+			}
+
+			if (RecenterType & RecenterOrientation)
+			{
+				Settings->BaseOrientation = FRotator(0, FRotator(ToFQuat(poseState.Pose.Orientation)).Yaw - Yaw, 0).Quaternion();
+			}
 		}
 	}
 
-
 	void FOculusHMD::SetBaseRotation(const FRotator& BaseRot)
 	{
 		SetBaseOrientation(BaseRot.Quaternion());
diff --git a/Engine/Plugins/Runtime/Oculus/OculusVR/Source/OculusHMD/Private/OculusHMD.h b/Engine/Plugins/Runtime/Oculus/OculusVR/Source/OculusHMD/Private/OculusHMD.h
--- a/Engine/Plugins/Runtime/Oculus/OculusVR/Source/OculusHMD/Private/OculusHMD.h
+++ b/Engine/Plugins/Runtime/Oculus/OculusVR/Source/OculusHMD/Private/OculusHMD.h
@@ -54,6 +54,12 @@ struct FPerformanceStats
 	}
 };
 
+enum FRecenterTypes
+{
+	RecenterOrientation = 0x1,
+	RecenterPosition = 0x2,
+	RecenterOrientationAndPosition = 0x3
+};
 
 //-------------------------------------------------------------------------------------------------
 // FOculusHMD - Oculus Rift Head Mounted Display
@@ -224,6 +230,7 @@ protected:
 	void ApplySystemOverridesOnStereo(bool force = false);
 	bool OnOculusStateChange(bool bIsEnabledNow);
 	bool ShouldDisableHiddenAndVisibileAreaMeshForSpectatorScreen_RenderThread() const;
+	void Recenter(FRecenterTypes RecenterType, float Yaw);
 #if !UE_BUILD_SHIPPING
 	void DrawDebug(UCanvas* InCanvas, APlayerController* InPlayerController);
 #endif
diff --git a/Engine/Plugins/Runtime/Oculus/OculusVR/Source/OculusHMD/Private/OculusHMD_Splash.cpp b/Engine/Plugins/Runtime/Oculus/OculusVR/Source/OculusHMD/Private/OculusHMD_Splash.cpp
--- a/Engine/Plugins/Runtime/Oculus/OculusVR/Source/OculusHMD/Private/OculusHMD_Splash.cpp
+++ b/Engine/Plugins/Runtime/Oculus/OculusVR/Source/OculusHMD/Private/OculusHMD_Splash.cpp
@@ -514,11 +514,7 @@ void FSplash::UnloadTexture(FSplashLayer& InSplashLayer)
 {
 	CheckInGameThread();
 
-	if (InSplashLayer.Desc.LoadingTexture && InSplashLayer.Desc.LoadingTexture->IsValidLowLevel())
-	{
-		InSplashLayer.Desc.LoadingTexture = nullptr;
-	}
-
+	InSplashLayer.Desc.LoadingTexture = nullptr;
 	InSplashLayer.Desc.LoadedTexture = nullptr;
 	InSplashLayer.Layer.Reset();
 }
