Oculus Touch: Gamepad Special Left/Right inputs swapped (copy of CL#4899121 from Dev-VR)

diff --git a/Engine/Plugins/Runtime/Oculus/OculusVR/Source/OculusInput/Private/OculusInputState.h b/Engine/Plugins/Runtime/Oculus/OculusVR/Source/OculusInput/Private/OculusInputState.h
--- a/Engine/Plugins/Runtime/Oculus/OculusVR/Source/OculusInput/Private/OculusInputState.h
+++ b/Engine/Plugins/Runtime/Oculus/OculusVR/Source/OculusInput/Private/OculusInputState.h
@@ -288,7 +288,7 @@ struct FOculusTouchControllerState
 		Buttons[ (int32)EOculusTouchControllerButton::Thumbstick_Left].Key = (Hand == EControllerHand::Left) ? FGamepadKeyNames::MotionController_Left_Thumbstick_Left : FGamepadKeyNames::MotionController_Right_Thumbstick_Left;
 		Buttons[ (int32)EOculusTouchControllerButton::Thumbstick_Right].Key = (Hand == EControllerHand::Left) ? FGamepadKeyNames::MotionController_Left_Thumbstick_Right : FGamepadKeyNames::MotionController_Right_Thumbstick_Right;
 
-		Buttons[ (int32)EOculusTouchControllerButton::Menu ].Key = FGamepadKeyNames::SpecialRight;
+		Buttons[ (int32)EOculusTouchControllerButton::Menu ].Key = (Hand == EControllerHand::Left) ? FGamepadKeyNames::SpecialLeft : FGamepadKeyNames::SpecialRight;
 
 		CapacitiveAxes[(int32)EOculusTouchCapacitiveAxes::Thumbstick].Axis = (Hand == EControllerHand::Left) ? FOculusKeyNames::OculusTouch_Left_Thumbstick : FOculusKeyNames::OculusTouch_Right_Thumbstick;
 		CapacitiveAxes[(int32)EOculusTouchCapacitiveAxes::Trigger].Axis = (Hand == EControllerHand::Left) ? FOculusKeyNames::OculusTouch_Left_Trigger : FOculusKeyNames::OculusTouch_Right_Trigger;
