Reenable 32bit depth.
Change OVR to also 32bit depth.

diff --git a/Engine/Plugins/Runtime/Oculus/OculusVR/Source/OculusHMD/Private/OculusHMD_CustomPresent.cpp b/Engine/Plugins/Runtime/Oculus/OculusVR/Source/OculusHMD/Private/OculusHMD_CustomPresent.cpp
--- a/Engine/Plugins/Runtime/Oculus/OculusVR/Source/OculusHMD/Private/OculusHMD_CustomPresent.cpp
+++ b/Engine/Plugins/Runtime/Oculus/OculusVR/Source/OculusHMD/Private/OculusHMD_CustomPresent.cpp
@@ -37,7 +37,7 @@ FCustomPresent::FCustomPresent(class FOculusHMD* InOculusHMD, ovrpRenderAPIType
 	CheckInGameThread();
 
 	DefaultOvrpTextureFormat = GetOvrpTextureFormat(GetDefaultPixelFormat());
-	DefaultDepthOvrpTextureFormat = bSupportsDepth ? ovrpTextureFormat_D24_S8 : ovrpTextureFormat_None;
+	DefaultDepthOvrpTextureFormat = bSupportsDepth ? ovrpTextureFormat_D32_S824_FP : ovrpTextureFormat_None;
 
 	// grab a pointer to the renderer module for displaying our mirror window
 	static const FName RendererModuleName("Renderer");
diff --git a/Engine/Source/Runtime/D3D12RHI/Private/D3D12RHI.cpp b/Engine/Source/Runtime/D3D12RHI/Private/D3D12RHI.cpp
--- a/Engine/Source/Runtime/D3D12RHI/Private/D3D12RHI.cpp
+++ b/Engine/Source/Runtime/D3D12RHI/Private/D3D12RHI.cpp
@@ -23,7 +23,7 @@ DEFINE_LOG_CATEGORY(LogD3D12RHI);
 
 static TAutoConsoleVariable<int32> CVarD3D12UseD24(
 	TEXT("r.D3D12.Depth24Bit"),
-	1,
+	0,
 	TEXT("0: Use 32-bit float depth buffer\n1: Use 24-bit fixed point depth buffer(default)\n"),
 	ECVF_ReadOnly
 );
diff --git a/Engine/Source/Runtime/Windows/D3D11RHI/Private/D3D11Device.cpp b/Engine/Source/Runtime/Windows/D3D11RHI/Private/D3D11Device.cpp
--- a/Engine/Source/Runtime/Windows/D3D11RHI/Private/D3D11Device.cpp
+++ b/Engine/Source/Runtime/Windows/D3D11RHI/Private/D3D11Device.cpp
@@ -33,7 +33,7 @@ IMPLEMENT_MODULE(FD3D11DynamicRHIModule, D3D11RHI);
 
 static TAutoConsoleVariable<int32> CVarD3D11UseD24(
 	TEXT("r.D3D11.Depth24Bit"),
-	1,
+	0,
 	TEXT("0: Use 32-bit float depth buffer\n1: Use 24-bit fixed point depth buffer(default)\n"),
 	ECVF_ReadOnly
 );
