Work around log spam due to AllocateDepthTexture being called every frame.
(Copy of CL#5039030 from Dev-VR)

diff --git a/Engine/Plugins/Runtime/Oculus/OculusVR/Source/OculusHMD/Private/OculusHMD.cpp b/Engine/Plugins/Runtime/Oculus/OculusVR/Source/OculusHMD/Private/OculusHMD.cpp
--- a/Engine/Plugins/Runtime/Oculus/OculusVR/Source/OculusHMD/Private/OculusHMD.cpp
+++ b/Engine/Plugins/Runtime/Oculus/OculusVR/Source/OculusHMD/Private/OculusHMD.cpp
@@ -1300,10 +1300,13 @@ namespace OculusHMD
 				// Ensure the texture size matches the eye layer. We may get other depth allocations unrelated to the main scene render.
 				if (FIntPoint(SizeX, SizeY) == TextureSet->GetTexture2D()->GetSizeXY())
 				{
-					UE_LOG(LogHMD, Log, TEXT("Allocating Oculus %d x %d depth rendertarget swapchain"), SizeX, SizeY);
 					OutTargetableTexture = TextureSet->GetTexture2D();
 					OutShaderResourceTexture = TextureSet->GetTexture2D();
-					bNeedReAllocateDepthTexture_RenderThread = false;
+					if (bNeedReAllocateDepthTexture_RenderThread)
+					{
+						UE_LOG(LogHMD, Log, TEXT("Allocating Oculus %d x %d depth rendertarget swapchain"), SizeX, SizeY);
+						bNeedReAllocateDepthTexture_RenderThread = false;
+					}
 					return true;
 				}
 			}
