Updating the Oculus plugin binaries to fix a known memory corruption issue.

diff --git a/Engine/Plugins/Runtime/Oculus/OculusVR/Source/OculusHMD/Private/OculusHMD.cpp b/Engine/Plugins/Runtime/Oculus/OculusVR/Source/OculusHMD/Private/OculusHMD.cpp
--- a/Engine/Plugins/Runtime/Oculus/OculusVR/Source/OculusHMD/Private/OculusHMD.cpp
+++ b/Engine/Plugins/Runtime/Oculus/OculusVR/Source/OculusHMD/Private/OculusHMD.cpp
@@ -2122,7 +2122,7 @@ namespace OculusHMD
 			if (MeshType == ovrpViewportStencilType_HiddenArea)
 			{
 				Vertex.Position.X = (Position.x * 2.0f) - 1.0f;
-				Vertex.Position.Y = (Position.y * 2.0f) - 1.0f;
+				Vertex.Position.Y = 1.0f - (Position.y * 2.0f);
 				Vertex.Position.Z = 1.0f;
 				Vertex.Position.W = 1.0f;
 				Vertex.UV.X = 0.0f;
@@ -2131,11 +2131,11 @@ namespace OculusHMD
 			else if (MeshType == ovrpViewportStencilType_VisibleArea)
 			{
 				Vertex.Position.X = Position.x;
-				Vertex.Position.Y = 1.0f - Position.y;
+				Vertex.Position.Y = Position.y;
 				Vertex.Position.Z = 0.0f;
 				Vertex.Position.W = 1.0f;
 				Vertex.UV.X = Position.x;
-				Vertex.UV.Y = 1.0f - Position.y;
+				Vertex.UV.Y = Position.y;
 			}
 			else
 			{
