Integrating a last minute fix from Oculus where CPU and GPU levels don't latch in UE4 if you go to VrShell (home button) and resume.

diff --git a/Engine/Plugins/Runtime/Oculus/OculusVR/Source/OculusHMD/Private/OculusFunctionLibrary.cpp b/Engine/Plugins/Runtime/Oculus/OculusVR/Source/OculusHMD/Private/OculusFunctionLibrary.cpp
--- a/Engine/Plugins/Runtime/Oculus/OculusVR/Source/OculusHMD/Private/OculusFunctionLibrary.cpp
+++ b/Engine/Plugins/Runtime/Oculus/OculusVR/Source/OculusHMD/Private/OculusFunctionLibrary.cpp
@@ -134,8 +134,7 @@ void UOculusFunctionLibrary::SetCPUAndGPULevels(int CPULevel, int GPULevel)
 	OculusHMD::FOculusHMD* OculusHMD = GetOculusHMD();
 	if (OculusHMD != nullptr && OculusHMD->IsHMDActive())
 	{
-		ovrp_SetSystemCpuLevel2(CPULevel);
-		ovrp_SetSystemGpuLevel2(GPULevel);
+		OculusHMD->SetCPUAndGPULevel(CPULevel, GPULevel);
 	}
 #endif // OCULUS_HMD_SUPPORTED_PLATFORMS
 }
diff --git a/Engine/Plugins/Runtime/Oculus/OculusVR/Source/OculusHMD/Private/OculusHMD.cpp b/Engine/Plugins/Runtime/Oculus/OculusVR/Source/OculusHMD/Private/OculusHMD.cpp
--- a/Engine/Plugins/Runtime/Oculus/OculusVR/Source/OculusHMD/Private/OculusHMD.cpp
+++ b/Engine/Plugins/Runtime/Oculus/OculusVR/Source/OculusHMD/Private/OculusHMD.cpp
@@ -2760,6 +2760,15 @@ namespace OculusHMD
 		return PerformanceStats;
 	}
 
+	void FOculusHMD::SetCPUAndGPULevel(int CPULevel, int GPULevel)
+	{
+		CheckInGameThread();
+		Settings->CPULevel = CPULevel;
+		Settings->GPULevel = GPULevel;
+		ovrp_SetSystemCpuLevel2(Settings->CPULevel);
+		ovrp_SetSystemGpuLevel2(Settings->GPULevel);
+	}
+
 
 	void FOculusHMD::SetTiledMultiResLevel(ETiledMultiResLevel multiresLevel)
 	{
diff --git a/Engine/Plugins/Runtime/Oculus/OculusVR/Source/OculusHMD/Private/OculusHMD.h b/Engine/Plugins/Runtime/Oculus/OculusVR/Source/OculusHMD/Private/OculusHMD.h
--- a/Engine/Plugins/Runtime/Oculus/OculusVR/Source/OculusHMD/Private/OculusHMD.h
+++ b/Engine/Plugins/Runtime/Oculus/OculusVR/Source/OculusHMD/Private/OculusHMD.h
@@ -328,6 +328,7 @@ public:
 	void StartRHIFrame_RenderThread(); // Called from PreRenderViewFamily_RenderThread
 	void FinishRHIFrame_RHIThread(); // Called from FinishRendering_RHIThread
 
+	void SetCPUAndGPULevel(int CPULevel, int GPULevel);
 	void SetTiledMultiResLevel(ETiledMultiResLevel multiresLevel);
 	void SetColorScaleAndOffset(FLinearColor ColorScale, FLinearColor ColorOffset, bool bApplyToAllLayers);
 
