[NO NON-PII COMMIT MESSAGE]

diff --git a/Engine/Config/BaseEngine.ini b/Engine/Config/BaseEngine.ini
--- a/Engine/Config/BaseEngine.ini
+++ b/Engine/Config/BaseEngine.ini
@@ -215,8 +215,6 @@ NetClientTicksPerSecond=200
 +StatColorMappings=(StatName="Streaming fudge factor",ColorMap=((In=0.0,Out=(G=255)),(In=1.0,Out=(G=255)),(In=2.5,Out=(R=255,G=255)),(In=5.0,Out=(R=255)),(In=10.0,Out=(R=255))))
 DisplayGamma=2.2
 MinDesiredFrameRate=35.000000
-InitialButtonRepeatDelay=0.2
-ButtonRepeatDelay=0.1
 NetDriverDefinitions=(DefName="GameNetDriver",DriverClassName="/Script/OnlineSubsystemUtils.IpNetDriver",DriverClassNameFallback="/Script/OnlineSubsystemUtils.IpNetDriver")
 +NetDriverDefinitions=(DefName="DemoNetDriver",DriverClassName="/Script/Engine.DemoNetDriver",DriverClassNameFallback="/Script/Engine.DemoNetDriver")
 NetErrorLogInterval=1.0
diff --git a/Engine/Config/BaseInput.ini b/Engine/Config/BaseInput.ini
--- a/Engine/Config/BaseInput.ini
+++ b/Engine/Config/BaseInput.ini
@@ -13,6 +13,8 @@ bEnableGestureRecognizer=false
 +ConsoleKeys=Tilde
 TriggerKeyIsConsideredPressed=0.75
 TriggerKeyIsConsideredReleased=0.25
+InitialButtonRepeatDelay=0.2
+ButtonRepeatDelay=0.1
 
 ;-----------------------------------------------------------------------------------------
 ; Axis properties
diff --git a/Engine/Plugins/Runtime/Oculus/OculusVR/Source/OculusInput/Private/OculusInput.cpp b/Engine/Plugins/Runtime/Oculus/OculusVR/Source/OculusInput/Private/OculusInput.cpp
--- a/Engine/Plugins/Runtime/Oculus/OculusVR/Source/OculusInput/Private/OculusInput.cpp
+++ b/Engine/Plugins/Runtime/Oculus/OculusVR/Source/OculusInput/Private/OculusInput.cpp
@@ -6,6 +6,7 @@
 #include "OculusHMD.h"
 #include "Misc/CoreDelegates.h"
 #include "Features/IModularFeatures.h"
+#include "Misc/ConfigCacheIni.h"
 
 #define OVR_DEBUG_LOGGING 0
 
@@ -78,6 +79,9 @@ float FOculusInput::TriggerThreshold = 0.8f;
 /** Are Remote keys mapped to gamepad or not. */
 bool FOculusInput::bRemoteKeysMappedToGamepad = true;
 
+float FOculusInput::InitialButtonRepeatDelay = 0.2f;
+float FOculusInput::ButtonRepeatDelay = 0.1f;
+
 FOculusInput::FOculusInput( const TSharedRef< FGenericApplicationMessageHandler >& InMessageHandler )
 	: OVRPluginHandle(nullptr)
 	, MessageHandler( InMessageHandler )
@@ -169,6 +173,9 @@ void FOculusInput::LoadConfig()
 	{
 		bRemoteKeysMappedToGamepad = b;
 	}
+
+	GConfig->GetFloat(TEXT("/Script/Engine.InputSettings"), TEXT("InitialButtonRepeatDelay"), InitialButtonRepeatDelay, GInputIni);
+	GConfig->GetFloat(TEXT("/Script/Engine.InputSettings"), TEXT("ButtonRepeatDelay"), ButtonRepeatDelay, GInputIni);
 }
 
 void FOculusInput::Tick( float DeltaTime )
@@ -180,10 +187,6 @@ void FOculusInput::Tick( float DeltaTime )
 void FOculusInput::SendControllerEvents()
 {
 	const double CurrentTime = FPlatformTime::Seconds();
-
-	// @todo: Should be made configurable and unified with other controllers handling of repeat
-	const float InitialButtonRepeatDelay = 0.2f;
-	const float ButtonRepeatDelay = 0.1f;
 	const float AnalogButtonPressThreshold = TriggerThreshold;
 
 	if(IOculusHMDModule::IsAvailable() && ovrp_GetInitialized() && FApp::HasVRFocus())
diff --git a/Engine/Plugins/Runtime/Oculus/OculusVR/Source/OculusInput/Private/OculusInput.h b/Engine/Plugins/Runtime/Oculus/OculusVR/Source/OculusInput/Private/OculusInput.h
--- a/Engine/Plugins/Runtime/Oculus/OculusVR/Source/OculusInput/Private/OculusInput.h
+++ b/Engine/Plugins/Runtime/Oculus/OculusVR/Source/OculusInput/Private/OculusInput.h
@@ -95,6 +95,10 @@ private:
 	/** Are Remote keys mapped to gamepad or not. */
 	static bool bRemoteKeysMappedToGamepad;
 
+	/** Repeat key delays, loaded from config */
+	static float InitialButtonRepeatDelay;
+	static float ButtonRepeatDelay;
+
 	ovrpHapticsDesc OvrpHapticsDesc;
 };
 
diff --git a/Engine/Plugins/Runtime/Steam/SteamController/Source/SteamController/Private/SteamController.cpp b/Engine/Plugins/Runtime/Steam/SteamController/Source/SteamController/Private/SteamController.cpp
--- a/Engine/Plugins/Runtime/Steam/SteamController/Source/SteamController/Private/SteamController.cpp
+++ b/Engine/Plugins/Runtime/Steam/SteamController/Source/SteamController/Private/SteamController.cpp
@@ -4,6 +4,7 @@
 #include "HAL/PlatformTime.h"
 #include "HAL/PlatformProcess.h"
 #include "Misc/Paths.h"
+#include "Misc/ConfigCacheIni.h"
 #include "GenericPlatform/GenericApplicationMessageHandler.h"
 #include "Modules/ModuleManager.h"
 #include "GenericPlatform/IInputInterface.h"
@@ -74,6 +75,9 @@ public:
 			return;
 		}
 
+		GConfig->GetDouble(TEXT("/Script/Engine.InputSettings"), TEXT("InitialButtonRepeatDelay"), InitialButtonRepeatDelay, GInputIni);
+		GConfig->GetDouble(TEXT("/Script/Engine.InputSettings"), TEXT("ButtonRepeatDelay"), ButtonRepeatDelay, GInputIni);
+
 		// Initialize the API, so we can start calling SteamController functions
 		bSteamAPIInitialized = SteamAPI_Init();
 
diff --git a/Engine/Plugins/Runtime/Steam/SteamVR/Source/SteamVRController/Private/SteamVRController.cpp b/Engine/Plugins/Runtime/Steam/SteamVR/Source/SteamVRController/Private/SteamVRController.cpp
--- a/Engine/Plugins/Runtime/Steam/SteamVR/Source/SteamVRController/Private/SteamVRController.cpp
+++ b/Engine/Plugins/Runtime/Steam/SteamVR/Source/SteamVRController/Private/SteamVRController.cpp
@@ -21,6 +21,7 @@
 #include "GameFramework/PlayerInput.h"
 
 #include "Misc/FileHelper.h"
+#include "Misc/ConfigCacheIni.h"
 #include "Serialization/JsonReader.h"
 #include "Serialization/JsonSerializer.h"
 
@@ -183,6 +184,9 @@ public:
 		InitialButtonRepeatDelay = 0.2f;
 		ButtonRepeatDelay = 0.1f;
 
+		GConfig->GetFloat(TEXT("/Script/Engine.InputSettings"), TEXT("InitialButtonRepeatDelay"), InitialButtonRepeatDelay, GInputIni);
+		GConfig->GetFloat(TEXT("/Script/Engine.InputSettings"), TEXT("ButtonRepeatDelay"), ButtonRepeatDelay, GInputIni);
+
 		InitControllerMappings();
 		InitLegacyControllerKeys();
 		BuildActionManifest();
diff --git a/Engine/Source/Runtime/ApplicationCore/Private/Android/AndroidInputInterface.cpp b/Engine/Source/Runtime/ApplicationCore/Private/Android/AndroidInputInterface.cpp
--- a/Engine/Source/Runtime/ApplicationCore/Private/Android/AndroidInputInterface.cpp
+++ b/Engine/Source/Runtime/ApplicationCore/Private/Android/AndroidInputInterface.cpp
@@ -105,6 +105,9 @@ FAndroidInputInterface::FAndroidInputInterface(const TSharedRef< FGenericApplica
 	InitialButtonRepeatDelay = 0.2f;
 	ButtonRepeatDelay = 0.1f;
 
+	GConfig->GetFloat(TEXT("/Script/Engine.InputSettings"), TEXT("InitialButtonRepeatDelay"), InitialButtonRepeatDelay, GInputIni);
+	GConfig->GetFloat(TEXT("/Script/Engine.InputSettings"), TEXT("ButtonRepeatDelay"), ButtonRepeatDelay, GInputIni);
+
 	VibeIsOn = false;
 
 	for (int32 DeviceIndex = 0; DeviceIndex < MAX_NUM_CONTROLLERS; DeviceIndex++)
diff --git a/Engine/Source/Runtime/ApplicationCore/Private/Lumin/LuminGamepadInterface.cpp b/Engine/Source/Runtime/ApplicationCore/Private/Lumin/LuminGamepadInterface.cpp
--- a/Engine/Source/Runtime/ApplicationCore/Private/Lumin/LuminGamepadInterface.cpp
+++ b/Engine/Source/Runtime/ApplicationCore/Private/Lumin/LuminGamepadInterface.cpp
@@ -3,6 +3,7 @@
 #include "LuminGamepadInterface.h"
 #include "Android/AndroidApplication.h"
 #include "GenericPlatform/GenericApplication.h"
+#include "Misc/ConfigCacheIni.h"
 #include "HAL/PlatformTime.h"
 #include <dirent.h>
 #include <errno.h>
@@ -76,6 +77,9 @@ FLuminGamepadInterface::FLuminGamepadInterface(const TSharedRef< FGenericApplica
 	InitialButtonRepeatDelay = 0.2f;
 	ButtonRepeatDelay = 0.1f;
 
+	GConfig->GetFloat(TEXT("/Script/Engine.InputSettings"), TEXT("InitialButtonRepeatDelay"), InitialButtonRepeatDelay, GInputIni);
+	GConfig->GetFloat(TEXT("/Script/Engine.InputSettings"), TEXT("ButtonRepeatDelay"), ButtonRepeatDelay, GInputIni);
+
 	// In the engine, all controllers map to xbox controllers for consistency 
 	X360ToXboxControllerMapping[0] = 0;		// A
 	X360ToXboxControllerMapping[1] = 1;		// B
diff --git a/Engine/Source/Runtime/ApplicationCore/Private/Mac/HIDInputInterface.cpp b/Engine/Source/Runtime/ApplicationCore/Private/Mac/HIDInputInterface.cpp
--- a/Engine/Source/Runtime/ApplicationCore/Private/Mac/HIDInputInterface.cpp
+++ b/Engine/Source/Runtime/ApplicationCore/Private/Mac/HIDInputInterface.cpp
@@ -3,6 +3,7 @@
 #include "HIDInputInterface.h"
 #include "HAL/PlatformTime.h"
 #include "Misc/CallbackDevice.h"
+#include "Misc/ConfigCacheIni.h"
 #include "Templates/SharedPointer.h"
 
 static int32 GetDevicePropertyAsInt32(IOHIDDeviceRef DeviceRef, CFStringRef Property)
@@ -37,6 +38,9 @@ HIDInputInterface::HIDInputInterface(const TSharedRef<FGenericApplicationMessage
 	InitialButtonRepeatDelay = 0.2f;
 	ButtonRepeatDelay = 0.1f;
 
+	GConfig->GetFloat(TEXT("/Script/Engine.InputSettings"), TEXT("InitialButtonRepeatDelay"), InitialButtonRepeatDelay, GInputIni);
+	GConfig->GetFloat(TEXT("/Script/Engine.InputSettings"), TEXT("ButtonRepeatDelay"), ButtonRepeatDelay, GInputIni);
+
 	Buttons[0] = FGamepadKeyNames::FaceButtonBottom;
 	Buttons[1] = FGamepadKeyNames::FaceButtonRight;
 	Buttons[2] = FGamepadKeyNames::FaceButtonLeft;
diff --git a/Engine/Source/Runtime/ApplicationCore/Private/Windows/XInputInterface.cpp b/Engine/Source/Runtime/ApplicationCore/Private/Windows/XInputInterface.cpp
--- a/Engine/Source/Runtime/ApplicationCore/Private/Windows/XInputInterface.cpp
+++ b/Engine/Source/Runtime/ApplicationCore/Private/Windows/XInputInterface.cpp
@@ -5,6 +5,7 @@
 #include "Misc/CoreDelegates.h"
 #include "Windows/WindowsApplication.h"
 #include "GenericPlatform/GenericApplication.h"
+#include "Misc/ConfigCacheIni.h"
 
 #pragma pack (push,8)
 #include "Windows/AllowWindowsPlatformTypes.h"
@@ -35,6 +36,9 @@ XInputInterface::XInputInterface( const TSharedRef< FGenericApplicationMessageHa
 	InitialButtonRepeatDelay = 0.2f;
 	ButtonRepeatDelay = 0.1f;
 
+	GConfig->GetFloat(TEXT("/Script/Engine.InputSettings"), TEXT("InitialButtonRepeatDelay"), InitialButtonRepeatDelay, GInputIni);
+	GConfig->GetFloat(TEXT("/Script/Engine.InputSettings"), TEXT("ButtonRepeatDelay"), ButtonRepeatDelay, GInputIni);
+
 	// In the engine, all controllers map to xbox controllers for consistency 
 	X360ToXboxControllerMapping[0] = 0;		// A
 	X360ToXboxControllerMapping[1] = 1;		// B
