Copy of CL#7633353. Work around an assert reported in a crash report.

diff --git a/Engine/Plugins/Runtime/Oculus/OculusVR/Source/OculusHMD/Private/OculusHMD.cpp b/Engine/Plugins/Runtime/Oculus/OculusVR/Source/OculusHMD/Private/OculusHMD.cpp
--- a/Engine/Plugins/Runtime/Oculus/OculusVR/Source/OculusHMD/Private/OculusHMD.cpp
+++ b/Engine/Plugins/Runtime/Oculus/OculusVR/Source/OculusHMD/Private/OculusHMD.cpp
@@ -1391,8 +1391,11 @@ namespace OculusHMD
 			InitializeEyeLayer_RenderThread(GetImmediateCommandList_ForRenderCommand());
 
 			const FXRSwapChainPtr& SwapChain = EyeLayer_RenderThread->GetSwapChain();
-			const FRHITexture2D* const SwapChainTexture = SwapChain->GetTexture2DArray() ? SwapChain->GetTexture2DArray() : SwapChain->GetTexture2D();
-			UE_LOG(LogHMD, Log, TEXT("Allocating Oculus %d x %d rendertarget swapchain"), SwapChainTexture->GetSizeX(), SwapChainTexture->GetSizeY());
+			if (SwapChain.IsValid())
+			{
+				const FRHITexture2D* const SwapChainTexture = SwapChain->GetTexture2DArray() ? SwapChain->GetTexture2DArray() : SwapChain->GetTexture2D();
+				UE_LOG(LogHMD, Log, TEXT("Allocating Oculus %d x %d rendertarget swapchain"), SwapChainTexture->GetSizeX(), SwapChainTexture->GetSizeY());
+			}
 		});
 
 		bNeedReAllocateViewportRenderTarget = true;
