Fix crash when starting PIE with VirtualCamera game mode when you have an oculus headset plugged in.
Activating the game mode changes the tracking system which doesn't have a HMD device

diff --git a/Engine/Plugins/Runtime/Oculus/OculusVR/Source/OculusInput/Private/OculusInput.cpp b/Engine/Plugins/Runtime/Oculus/OculusVR/Source/OculusInput/Private/OculusInput.cpp
--- a/Engine/Plugins/Runtime/Oculus/OculusVR/Source/OculusInput/Private/OculusInput.cpp
+++ b/Engine/Plugins/Runtime/Oculus/OculusVR/Source/OculusInput/Private/OculusInput.cpp
@@ -191,7 +191,7 @@ void FOculusInput::SendControllerEvents()
 
 	if(IOculusHMDModule::IsAvailable() && ovrp_GetInitialized() && FApp::HasVRFocus())
 	{
-		if (MessageHandler.IsValid())
+		if (MessageHandler.IsValid() && GEngine->XRSystem->GetHMDDevice())
 		{
 			OculusHMD::FOculusHMD* OculusHMD = static_cast<OculusHMD::FOculusHMD*>(GEngine->XRSystem->GetHMDDevice());
 			ovrp_Update3(ovrpStep_Render, OculusHMD->GetNextFrameNumber(), 0.0);
