Add check and log errors for android/Oculus to ensure if mobile multi-view is enabled, multiview direct is also enabled, and also log an error if mobile HDR is enabled (this is not supported on Oculus).

diff --git a/Engine/Plugins/Runtime/Oculus/OculusVR/Source/OculusHMD/Private/OculusHMD.cpp b/Engine/Plugins/Runtime/Oculus/OculusVR/Source/OculusHMD/Private/OculusHMD.cpp
--- a/Engine/Plugins/Runtime/Oculus/OculusVR/Source/OculusHMD/Private/OculusHMD.cpp
+++ b/Engine/Plugins/Runtime/Oculus/OculusVR/Source/OculusHMD/Private/OculusHMD.cpp
@@ -2419,9 +2419,26 @@ namespace OculusHMD
 #if PLATFORM_ANDROID
 		static const auto CVarMobileMultiView = IConsoleManager::Get().FindTConsoleVariableDataInt(TEXT("vr.MobileMultiView"));
 		static const auto CVarMobileMultiViewDirect = IConsoleManager::Get().FindTConsoleVariableDataInt(TEXT("vr.MobileMultiView.Direct"));
+		static const auto CVarMobileHDR = IConsoleManager::Get().FindTConsoleVariableDataInt(TEXT("r.MobileHDR"));
 		const bool bIsMobileMultiViewEnabled = (CVarMobileMultiView && CVarMobileMultiView->GetValueOnAnyThread() != 0);
 		const bool bIsMobileMultiViewDirectEnabled = (CVarMobileMultiViewDirect && CVarMobileMultiViewDirect->GetValueOnAnyThread() != 0);
 		const bool bIsUsingDirectMobileMultiView = GSupportsMobileMultiView && bIsMobileMultiViewEnabled && bIsMobileMultiViewDirectEnabled;
+		const bool bMobileHDR = CVarMobileHDR && CVarMobileHDR->GetValueOnAnyThread() == 1;
+
+		if (bIsMobileMultiViewEnabled && !bIsMobileMultiViewDirectEnabled)
+		{
+			static bool bDisplayedMultiViewError = false;
+			UE_CLOG(!bDisplayedMultiViewError, LogHMD, Error, TEXT("\"Mobile Multiview Direct\" must always be enabled if \"Mobile Multiview\" is enabled on Oculus Mobile HMD devices."));
+			bDisplayedMultiViewError = true;
+		}
+
+		if (bMobileHDR)
+		{
+			static bool bDisplayedHDRError = false;
+			UE_CLOG(!bDisplayedHDRError, LogHMD, Error, TEXT("Mobile HDR is not supported on Oculus Mobile HMD devices."));
+			bDisplayedHDRError = true;
+		}
+
 		if (Settings->Flags.bDirectMultiview && bIsUsingDirectMobileMultiView)
 		{
 			Layout = ovrpLayout_Array;
