Fix for crash in UE4Editor-OculusEditor!SOculusPlatformToolWidget::SOculusPlatformToolWidget() [OculusPlatformToolWidget.cpp:30] - this crash was caused by a missing OVRPlugin.dll module. Added code to verify this module exists before populating UI components that rely on it.

diff --git a/Engine/Plugins/Runtime/Oculus/OculusVR/Source/OculusEditor/Private/OculusEditorModule.cpp b/Engine/Plugins/Runtime/Oculus/OculusVR/Source/OculusEditor/Private/OculusEditorModule.cpp
--- a/Engine/Plugins/Runtime/Oculus/OculusVR/Source/OculusEditor/Private/OculusEditorModule.cpp
+++ b/Engine/Plugins/Runtime/Oculus/OculusVR/Source/OculusEditor/Private/OculusEditorModule.cpp
@@ -33,6 +33,13 @@ void FOculusEditorModule::PostLoadCallback()
 
 void FOculusEditorModule::StartupModule()
 {
+	void* ModuleCheck = FPlatformProcess::GetDllHandle(TEXT("OVRPlugin.dll"));
+	if (!ModuleCheck)
+	{
+		return;
+	}
+
+	bModuleValid = true;
 	RegisterSettings();
 	FOculusAssetDirectory::LoadForCook();
 
@@ -76,6 +83,11 @@ void FOculusEditorModule::StartupModule()
 
 void FOculusEditorModule::ShutdownModule()
 {
+	if (!bModuleValid)
+	{
+		return;
+	}
+
 	if(!IsRunningCommandlet())
 	{
 		FOculusToolStyle::Shutdown();
diff --git a/Engine/Plugins/Runtime/Oculus/OculusVR/Source/OculusEditor/Private/OculusEditorModule.h b/Engine/Plugins/Runtime/Oculus/OculusVR/Source/OculusEditor/Private/OculusEditorModule.h
--- a/Engine/Plugins/Runtime/Oculus/OculusVR/Source/OculusEditor/Private/OculusEditorModule.h
+++ b/Engine/Plugins/Runtime/Oculus/OculusVR/Source/OculusEditor/Private/OculusEditorModule.h
@@ -17,7 +17,7 @@ class FMenuBuilder;
 class FOculusEditorModule : public IOculusEditorModule
 {
 public:
-	FOculusEditorModule() {};
+	FOculusEditorModule() : bModuleValid(false) {};
 
 	/** IModuleInterface implementation */
 	virtual void StartupModule() override;
@@ -44,6 +44,7 @@ private:
 
 private:
 	TSharedPtr<class FUICommandList> PluginCommands;
+	bool bModuleValid;
 };
 
 class IDetailLayoutBuilder;
