Make sure to clean up the temporary module handle used to check for OVRPlugin.dll being present.

diff --git a/Engine/Plugins/Runtime/Oculus/OculusVR/Source/OculusEditor/Private/OculusEditorModule.cpp b/Engine/Plugins/Runtime/Oculus/OculusVR/Source/OculusEditor/Private/OculusEditorModule.cpp
--- a/Engine/Plugins/Runtime/Oculus/OculusVR/Source/OculusEditor/Private/OculusEditorModule.cpp
+++ b/Engine/Plugins/Runtime/Oculus/OculusVR/Source/OculusEditor/Private/OculusEditorModule.cpp
@@ -39,6 +39,9 @@ void FOculusEditorModule::StartupModule()
 		return;
 	}
 
+	FPlatformProcess::FreeDllHandle(ModuleCheck);
+	ModuleCheck = nullptr;
+
 	bModuleValid = true;
 	RegisterSettings();
 	FOculusAssetDirectory::LoadForCook();
