Cleaning up questionable partner code found by SA.
I'm pretty confident the SA warning is a false positive, but there was a lot to fix here regardless.

diff --git a/Engine/Plugins/Runtime/Oculus/OculusVR/Source/OculusHMD/Private/OculusHMD.cpp b/Engine/Plugins/Runtime/Oculus/OculusVR/Source/OculusHMD/Private/OculusHMD.cpp
--- a/Engine/Plugins/Runtime/Oculus/OculusVR/Source/OculusHMD/Private/OculusHMD.cpp
+++ b/Engine/Plugins/Runtime/Oculus/OculusVR/Source/OculusHMD/Private/OculusHMD.cpp
@@ -2104,24 +2104,27 @@ namespace OculusHMD
 
 	void BuildOcclusionMesh(FHMDViewMesh& Mesh, ovrpEye Eye, ovrpViewportStencilType MeshType)
 	{
-		int VertexCount, IndexCount;
-		ovrpResult result;
-		if (OVRP_FAILURE(result = ovrp_GetViewportStencil(Eye, MeshType, nullptr, &VertexCount, nullptr, &IndexCount)))
+		int VertexCount = 0;
+		int IndexCount = 0;
+
+		ovrpResult Result = ovrpResult::ovrpFailure;
+		if (OVRP_FAILURE(Result = ovrp_GetViewportStencil(Eye, MeshType, nullptr, &VertexCount, nullptr, &IndexCount)))
 		{
 			return;
 		}
 
 		Mesh.pVertices = new FFilterVertex[VertexCount];
-		ovrpVector2f* ovrpVertices = new ovrpVector2f[VertexCount];
+		ovrpVector2f* const ovrpVertices = new ovrpVector2f[VertexCount];
 
 		Mesh.pIndices = new uint16[IndexCount];
 
 		ovrp_GetViewportStencil(Eye, MeshType, ovrpVertices, &VertexCount, Mesh.pIndices, &IndexCount);
 
-		for (int i = 0; i < IndexCount; i++)
+		for (int i = 0; i < VertexCount; ++i)
 		{
 			FFilterVertex& Vertex = Mesh.pVertices[i];
-			ovrpVector2f& Position = ovrpVertices[i];
+			CA_SUPPRESS(6385); //  warning C6385: Reading invalid data from 'ovrpVertices':  the readable size is 'VertexCount*8' bytes, but '16' bytes may be read
+			const ovrpVector2f& Position = ovrpVertices[i];
 			if (MeshType == ovrpViewportStencilType_HiddenArea)
 			{
 				Vertex.Position.X = (Position.x * 2.0f) - 1.0f;
@@ -2140,12 +2143,17 @@ namespace OculusHMD
 				Vertex.UV.X = Position.x;
 				Vertex.UV.Y = 1.0f - Position.y;
 			}
-
+			else
+			{
+				check(0);
+			}
 		}
 
 		Mesh.NumIndices = IndexCount;
 		Mesh.NumVertices = VertexCount;
 		Mesh.NumTriangles = IndexCount / 3;
+
+		delete [] ovrpVertices;
 	}
 
 	void FOculusHMD::SetupOcclusionMeshes()
