Fix OculusHMD plugin linking to the Windows Vulkan libraries on Android platforms

diff --git a/Engine/Plugins/Runtime/Oculus/OculusVR/Source/OculusHMD/OculusHMD.Build.cs b/Engine/Plugins/Runtime/Oculus/OculusVR/Source/OculusHMD/OculusHMD.Build.cs
--- a/Engine/Plugins/Runtime/Oculus/OculusVR/Source/OculusHMD/OculusHMD.Build.cs
+++ b/Engine/Plugins/Runtime/Oculus/OculusVR/Source/OculusHMD/OculusHMD.Build.cs
@@ -137,42 +137,31 @@ namespace UnrealBuildTool.Rules
 			}
 			else if (Target.Platform == UnrealTargetPlatform.Android)
 			{
-                // Vulkan
-                {
-                    string VulkanSDKPath = Environment.GetEnvironmentVariable("VULKAN_SDK");
-                    bool bSDKInstalled = !String.IsNullOrEmpty(VulkanSDKPath);
-                    bool bUseThirdParty = true;
-                    if (bSDKInstalled)
-                    {
-                        // Check if the installed SDK is newer or the same than the provided headers distributed with the Engine
-                        int ThirdPartyVersion = GetThirdPartyVersion();
-                        int SDKVersion = GetSDKVersion(VulkanSDKPath);
-                        if (SDKVersion >= ThirdPartyVersion)
-                        {
-                            // If the user has an installed SDK, use that instead
-                            PrivateIncludePaths.Add(VulkanSDKPath + "/Include");
-                            // Older SDKs have an extra subfolder
-                            PrivateIncludePaths.Add(VulkanSDKPath + "/Include/vulkan");
-
-                            if (Target.Platform == UnrealTargetPlatform.Win32)
-                            {
-                                PublicLibraryPaths.Add(VulkanSDKPath + "/Source/lib32");
-                            }
-                            else
-                            {
-                                PublicLibraryPaths.Add(VulkanSDKPath + "/Source/lib");
-                            }
-
-                            PublicAdditionalLibraries.Add("vulkan-1.lib");
-                            PublicAdditionalLibraries.Add("vkstatic.1.lib");
-                            bUseThirdParty = false;
-                        }
-                    }
-                    if (bUseThirdParty)
-                    {
-                        AddEngineThirdPartyPrivateStaticDependencies(Target, "Vulkan");
-                    }
-                }
+				// Vulkan
+				{
+					string VulkanSDKPath = Environment.GetEnvironmentVariable("VULKAN_SDK");
+					bool bSDKInstalled = !String.IsNullOrEmpty(VulkanSDKPath);
+					bool bUseThirdParty = true;
+					if (bSDKInstalled)
+					{
+						// Check if the installed SDK is newer or the same than the provided headers distributed with the Engine
+						int ThirdPartyVersion = GetThirdPartyVersion();
+						int SDKVersion = GetSDKVersion(VulkanSDKPath);
+						if (SDKVersion >= ThirdPartyVersion)
+						{
+							// If the user has an installed SDK, use that instead
+							PrivateIncludePaths.Add(VulkanSDKPath + "/Include");
+							// Older SDKs have an extra subfolder
+							PrivateIncludePaths.Add(VulkanSDKPath + "/Include/vulkan");
+
+							bUseThirdParty = false;
+						}
+					}
+					if (bUseThirdParty)
+					{
+						AddEngineThirdPartyPrivateStaticDependencies(Target, "Vulkan");
+					}
+				}
 
 				// AndroidPlugin
 				{
