Fix redundant checks of IsOnRHIThread when calling ExecuteOnRHIThread

diff --git a/Engine/Plugins/Runtime/Oculus/OculusVR/Source/OculusHMD/Private/OculusHMD_Layer.cpp b/Engine/Plugins/Runtime/Oculus/OculusVR/Source/OculusHMD/Private/OculusHMD_Layer.cpp
--- a/Engine/Plugins/Runtime/Oculus/OculusVR/Source/OculusHMD/Private/OculusHMD_Layer.cpp
+++ b/Engine/Plugins/Runtime/Oculus/OculusVR/Source/OculusHMD/Private/OculusHMD_Layer.cpp
@@ -32,17 +32,11 @@ FOvrpLayer::FOvrpLayer(uint32 InOvrpLayerId) :
 
 FOvrpLayer::~FOvrpLayer()
 {
-	if (InRenderThread())
-	{
-		ExecuteOnRHIThread_DoNotWait([this]()
-		{
-			ovrp_DestroyLayer(OvrpLayerId);
-		});
-	}
-	else
+	check(InRenderThread() || InRHIThread());
+	ExecuteOnRHIThread_DoNotWait([this]()
 	{
 		ovrp_DestroyLayer(OvrpLayerId);
-	}
+	});
 }
 
 
diff --git a/Engine/Plugins/Runtime/Oculus/OculusVR/Source/OculusHMD/Private/OculusHMD_TextureSetProxy.cpp b/Engine/Plugins/Runtime/Oculus/OculusVR/Source/OculusHMD/Private/OculusHMD_TextureSetProxy.cpp
--- a/Engine/Plugins/Runtime/Oculus/OculusVR/Source/OculusHMD/Private/OculusHMD_TextureSetProxy.cpp
+++ b/Engine/Plugins/Runtime/Oculus/OculusVR/Source/OculusHMD/Private/OculusHMD_TextureSetProxy.cpp
@@ -22,17 +22,11 @@ FTextureSetProxy::FTextureSetProxy(FTextureRHIParamRef InRHITexture, const TArra
 
 FTextureSetProxy::~FTextureSetProxy()
 {
-	if (InRenderThread())
-	{
-		ExecuteOnRHIThread([this]()
-		{
-			ReleaseResources_RHIThread();
-		});
-	}
-	else
+	check(InRenderThread() || InRHIThread());
+	ExecuteOnRHIThread([this]()
 	{
 		ReleaseResources_RHIThread();
-	}
+	});
 }
 
 
