Allow enforcing tracking when stereo rendering is disabled.
Tested on SteamVR (only required changes in FHeadMountedDisplayBase)
Oculus already existing code adapted to use the new interface.

diff --git a/Engine/Plugins/Runtime/Oculus/OculusVR/Source/OculusHMD/Private/OculusHMD.cpp b/Engine/Plugins/Runtime/Oculus/OculusVR/Source/OculusHMD/Private/OculusHMD.cpp
--- a/Engine/Plugins/Runtime/Oculus/OculusVR/Source/OculusHMD/Private/OculusHMD.cpp
+++ b/Engine/Plugins/Runtime/Oculus/OculusVR/Source/OculusHMD/Private/OculusHMD.cpp
@@ -446,25 +446,47 @@ namespace OculusHMD
 	}
 
 
-	bool FOculusHMD::IsHeadTrackingAllowed() const
+	bool FOculusHMD::IsHeadTrackingEnforced() const
+	{
+		if (IsInGameThread())
+		{
+			return Settings.IsValid() && Settings->Flags.bHeadTrackingEnforced;
+		}
+		else
+		{
+			CheckInRenderThread();
+			return Settings_RenderThread.IsValid() && Settings_RenderThread->Flags.bHeadTrackingEnforced;
+		}
+	}
+
+	void FOculusHMD::SetHeadTrackingEnforced(bool bEnabled)
 	{
 		CheckInGameThread();
+		check(Settings.IsValid());
 
-		if (!ovrp_GetInitialized())
+		const bool bOldValue = Settings->Flags.bHeadTrackingEnforced;
+		Settings->Flags.bHeadTrackingEnforced = bEnabled;
+
+		if (!bEnabled)
 		{
-			return false;
+			ResetControlRotation();
+		}
+		else if (!bOldValue)
+		{
+			InitDevice();
 		}
+	}
 
-#if WITH_EDITOR
-		if (GIsEditor)
+	bool FOculusHMD::IsHeadTrackingAllowed() const
+	{
+		CheckInGameThread();
+
+		if (!ovrp_GetInitialized())
 		{
-			// @todo vreditor: We need to do a pass over VREditor code and make sure we are handling the VR modes correctly.  HeadTracking can be enabled without Stereo3D, for example
-			UEditorEngine* EdEngine = Cast<UEditorEngine>(GEngine);
-			return (!EdEngine || EdEngine->IsHMDTrackingAllowed()) && (Settings->Flags.bHeadTrackingEnforced || GEngine->IsStereoscopic3D());
+			return false;
 		}
-#endif//WITH_EDITOR
 
-		return Settings.IsValid() && (Settings->Flags.bHeadTrackingEnforced || Settings->IsStereoEnabled());
+		return FHeadMountedDisplayBase::IsHeadTrackingAllowed();
 	}
 
 
@@ -1131,7 +1153,7 @@ namespace OculusHMD
 	{
 		CheckInRenderThread();
 
-		check(IsStereoEnabled());
+		check(IsStereoEnabled() || IsHeadTrackingEnforced());
 
 		// Don't use GetStereoProjectionMatrix because it is game thread only on oculus, we also don't need the zplane adjustments for this.
 		const int32 ViewIndex = GetViewIndexForPass(StereoPassType);
@@ -3275,30 +3297,6 @@ namespace OculusHMD
 	}
 
 #if !UE_BUILD_SHIPPING
-	void FOculusHMD::EnforceHeadTrackingCommandHandler(const TArray<FString>& Args, UWorld* World, FOutputDevice& Ar)
-	{
-		CheckInGameThread();
-
-		bool bOldValue = Settings->Flags.bHeadTrackingEnforced;
-
-		if (Args.Num() > 0)
-		{
-			Settings->Flags.bHeadTrackingEnforced = Args[0].Equals(TEXT("toggle"), ESearchCase::IgnoreCase) ? !Settings->Flags.bHeadTrackingEnforced : FCString::ToBool(*Args[0]);
-			if (!Settings->Flags.bHeadTrackingEnforced)
-			{
-				ResetControlRotation();
-			}
-		}
-
-		Ar.Logf(TEXT("Enforced head tracking is %s"), Settings->Flags.bHeadTrackingEnforced ? TEXT("on") : TEXT("off"));
-
-		if (!bOldValue && Settings->Flags.bHeadTrackingEnforced)
-		{
-			InitDevice();
-		}
-	}
-
-
 	void FOculusHMD::StatsCommandHandler(const TArray<FString>& Args, UWorld*, FOutputDevice& Ar)
 	{
 		CheckInGameThread();
diff --git a/Engine/Plugins/Runtime/Oculus/OculusVR/Source/OculusHMD/Private/OculusHMD.h b/Engine/Plugins/Runtime/Oculus/OculusVR/Source/OculusHMD/Private/OculusHMD.h
--- a/Engine/Plugins/Runtime/Oculus/OculusVR/Source/OculusHMD/Private/OculusHMD.h
+++ b/Engine/Plugins/Runtime/Oculus/OculusVR/Source/OculusHMD/Private/OculusHMD.h
@@ -94,6 +94,8 @@ public:
 	virtual class IHeadMountedDisplay* GetHMDDevice() override { return this; }
 	virtual class TSharedPtr< class IStereoRendering, ESPMode::ThreadSafe > GetStereoRenderingDevice() override { return SharedThis(this); }
 	//virtual class IXRInput* GetXRInput() override;
+	virtual bool IsHeadTrackingEnforced() const override;
+	virtual void SetHeadTrackingEnforced(bool bEnabled) override;
 	virtual bool IsHeadTrackingAllowed() const override;
 	virtual void OnBeginPlay(FWorldContext& InWorldContext) override;
 	virtual void OnEndPlay(FWorldContext& InWorldContext) override;
@@ -330,7 +332,6 @@ protected:
 	void ShowGlobalMenuCommandHandler(const TArray<FString>& Args, UWorld* World, FOutputDevice& Ar);
 	void ShowQuitMenuCommandHandler(const TArray<FString>& Args, UWorld* World, FOutputDevice& Ar);
 #if !UE_BUILD_SHIPPING
-	void EnforceHeadTrackingCommandHandler(const TArray<FString>& Args, UWorld* World, FOutputDevice& Ar);
 	void StatsCommandHandler(const TArray<FString>& Args, UWorld* World, FOutputDevice& Ar);
 	void ShowSettingsCommandHandler(const TArray<FString>& Args, UWorld* World, FOutputDevice& Ar);
 	void IPDCommandHandler(const TArray<FString>& Args, UWorld* World, FOutputDevice& Ar);
diff --git a/Engine/Plugins/Runtime/Oculus/OculusVR/Source/OculusHMD/Private/OculusHMD_ConsoleCommands.cpp b/Engine/Plugins/Runtime/Oculus/OculusVR/Source/OculusHMD/Private/OculusHMD_ConsoleCommands.cpp
--- a/Engine/Plugins/Runtime/Oculus/OculusVR/Source/OculusHMD/Private/OculusHMD_ConsoleCommands.cpp
+++ b/Engine/Plugins/Runtime/Oculus/OculusVR/Source/OculusHMD/Private/OculusHMD_ConsoleCommands.cpp
@@ -45,10 +45,6 @@ FConsoleCommands::FConsoleCommands(class FOculusHMD* InHMDPtr)
 		FConsoleCommandWithWorldArgsAndOutputDeviceDelegate::CreateRaw(InHMDPtr, &FOculusHMD::ShowQuitMenuCommandHandler))
 
 #if !UE_BUILD_SHIPPING
-	, EnforceHeadTrackingCommand(TEXT("vr.oculus.Debug.EnforceHeadTracking"),
-		*NSLOCTEXT("OculusRift", "CCommandText_EnforceTracking",
-			"Oculus Rift specific extension.\nSet to on to enforce head tracking even when not in stereo mode.").ToString(),
-		FConsoleCommandWithWorldArgsAndOutputDeviceDelegate::CreateRaw(InHMDPtr, &FOculusHMD::EnforceHeadTrackingCommandHandler))
 	, StatsCommand(TEXT("vr.oculus.Debug.bShowStats"),
 		*NSLOCTEXT("OculusRift", "CCommandText_Stats",
 			"Oculus Rift specific extension.\nEnable or disable rendering of stats.").ToString(),
@@ -82,6 +78,12 @@ bool FConsoleCommands::Exec(UWorld* InWorld, const TCHAR* Cmd, FOutputDevice& Ar
 	{
 		AliasedCommand = TEXT("vr.oculus.ShowQuitMenu");
 	}
+#if !UE_BUILD_SHIPPING
+	else if (FParse::Command(&Cmd, TEXT("vr.oculus.Debug.EnforceHeadTracking")))
+	{
+		AliasedCommand = TEXT("vr.HeadTracking.bEnforced");
+	}
+#endif // !UE_BUILD_SHIPPING
 
 	if (!AliasedCommand.IsEmpty())
 	{
diff --git a/Engine/Plugins/Runtime/Oculus/OculusVR/Source/OculusHMD/Private/OculusHMD_ConsoleCommands.h b/Engine/Plugins/Runtime/Oculus/OculusVR/Source/OculusHMD/Private/OculusHMD_ConsoleCommands.h
--- a/Engine/Plugins/Runtime/Oculus/OculusVR/Source/OculusHMD/Private/OculusHMD_ConsoleCommands.h
+++ b/Engine/Plugins/Runtime/Oculus/OculusVR/Source/OculusHMD/Private/OculusHMD_ConsoleCommands.h
@@ -31,7 +31,6 @@ private:
 
 #if !UE_BUILD_SHIPPING
 	// Debug console commands
-	FAutoConsoleCommand EnforceHeadTrackingCommand;
 	FAutoConsoleCommand StatsCommand;
 	FAutoConsoleCommand CubemapCommand;
 	FAutoConsoleCommand ShowSettingsCommand;
diff --git a/Engine/Source/Runtime/HeadMountedDisplay/Private/HeadMountedDisplayBase.cpp b/Engine/Source/Runtime/HeadMountedDisplay/Private/HeadMountedDisplayBase.cpp
--- a/Engine/Source/Runtime/HeadMountedDisplay/Private/HeadMountedDisplayBase.cpp
+++ b/Engine/Source/Runtime/HeadMountedDisplay/Private/HeadMountedDisplayBase.cpp
@@ -53,17 +53,28 @@ bool FHeadMountedDisplayBase::PopulateAnalyticsAttributes(TArray<FAnalyticsEvent
 	return true;
 }
 
+bool FHeadMountedDisplayBase::IsHeadTrackingEnforced() const
+{
+	return bHeadTrackingEnforced;
+}
+
+void FHeadMountedDisplayBase::SetHeadTrackingEnforced(bool bEnabled)
+{
+	bHeadTrackingEnforced = bEnabled;
+}
+
 bool FHeadMountedDisplayBase::IsHeadTrackingAllowed() const
 {
+	const bool bTrackingEnabled = IsStereoEnabled() || IsHeadTrackingEnforced();
 #if WITH_EDITOR
 	if (GIsEditor)
 	{
 		// @todo vreditor: We need to do a pass over VREditor code and make sure we are handling the VR modes correctly.  HeadTracking can be enabled without Stereo3D, for example
 		UEditorEngine* EdEngine = Cast<UEditorEngine>(GEngine);
-		return (!EdEngine || EdEngine->IsHMDTrackingAllowed()) && IsStereoEnabled();
+		return (!EdEngine || EdEngine->IsHMDTrackingAllowed()) && bTrackingEnabled;
 	}
 #endif // WITH_EDITOR
-	return IsStereoEnabled();
+	return bTrackingEnabled;
 }
 
 IStereoLayers* FHeadMountedDisplayBase::GetStereoLayers()
@@ -84,7 +95,9 @@ FVector2D FHeadMountedDisplayBase::GetEyeCenterPoint_RenderThread(EStereoscopicP
 {
 	check(IsInRenderingThread());
 
-	if (!IsStereoEnabled())
+	// Note: IsHeadTrackingAllowed() can only be called from the game thread.
+	// IsStereoEnabled() and IsHeadTrackingEnforced() can be called from both the render and game threads, however.
+	if (IsStereoEnabled() || IsHeadTrackingEnforced())
 	{
 		return FVector2D(0.5f, 0.5f);
 	}
diff --git a/Engine/Source/Runtime/HeadMountedDisplay/Private/HeadMountedDisplayConsoleCommands.cpp b/Engine/Source/Runtime/HeadMountedDisplay/Private/HeadMountedDisplayConsoleCommands.cpp
--- a/Engine/Source/Runtime/HeadMountedDisplay/Private/HeadMountedDisplayConsoleCommands.cpp
+++ b/Engine/Source/Runtime/HeadMountedDisplay/Private/HeadMountedDisplayConsoleCommands.cpp
@@ -341,7 +341,7 @@ static void HMDStatus(const TArray<FString>& Args, UWorld*, FOutputDevice& Ar)
 		auto HMD = GEngine->XRSystem;
 		Ar.Logf(TEXT("Position tracking status: %s\nHead tracking allowed: %s\nNumber of tracking sensors: %d"), 
 			HMD->DoesSupportPositionalTracking() ? (HMD->HasValidTrackingPosition() ? TEXT("active") : TEXT("lost")) : TEXT("not supported"),
-			HMD->IsHeadTrackingAllowed() ? TEXT("yes") : TEXT("no"),
+			HMD->IsHeadTrackingAllowed() ? (HMD->IsHeadTrackingEnforced() ? TEXT("enforced") : TEXT("yes")) : TEXT("no"),
 			HMD->CountTrackedDevices(EXRTrackedDeviceType::TrackingReference));
 	}
 }
@@ -351,6 +351,32 @@ static FAutoConsoleCommand CHMDStatusCmd(
 	*LOCTEXT("CVarText_HMDStatus", "Reports the current status of the head tracking.").ToString(),
 	FConsoleCommandWithWorldArgsAndOutputDeviceDelegate::CreateStatic(HMDStatus));
 
+static void ForceTracking(const TArray<FString>& Args, UWorld*, FOutputDevice& Ar)
+{
+	bool bEnable = false;
+	if (!GEngine || !GEngine->XRSystem.IsValid())
+	{
+		return;
+	}
+
+	if (Args.Num())
+	{
+		bEnable = FCString::ToBool(*Args[0]);
+		GEngine->XRSystem->SetHeadTrackingEnforced(bEnable);
+	}
+	else
+	{
+		bEnable = GEngine->XRSystem->IsHeadTrackingEnforced();
+		Ar.Logf(TEXT("Forced head tracking is %s"), bEnable ? TEXT("on") : TEXT("off"));
+	}
+}
+
+static FAutoConsoleCommand CForceTrackingCmd(
+	TEXT("vr.HeadTracking.bEnforced"),
+	*LOCTEXT("CCommandText_ForceTracking", "If set, head tracking is enabled even when stereo rendering is disabled.\nMay not be supported by all XR implementations.").ToString(),
+	FConsoleCommandWithWorldArgsAndOutputDeviceDelegate::CreateStatic(ForceTracking));
+
+
 static void EnableHMD(const TArray<FString>& Args, UWorld*, FOutputDevice& Ar)
 {
 	bool bEnable = 0;
@@ -417,6 +443,7 @@ static FAutoConsoleCommand CEnableStereoCmd(
 	*LOCTEXT("CCommandText_EnableStereo", "Enables or disables the stereo rendering. Use 1, True, or Yes to enable, 0, False or No to disable.").ToString(),
 	FConsoleCommandWithWorldArgsAndOutputDeviceDelegate::CreateStatic(EnableStereo));
 
+
 static void HMDVersion(const TArray<FString>& Args, UWorld*, FOutputDevice& Ar)
 {
 	if (GEngine && GEngine->XRSystem.IsValid())
diff --git a/Engine/Source/Runtime/HeadMountedDisplay/Public/HeadMountedDisplayBase.h b/Engine/Source/Runtime/HeadMountedDisplay/Public/HeadMountedDisplayBase.h
--- a/Engine/Source/Runtime/HeadMountedDisplay/Public/HeadMountedDisplayBase.h
+++ b/Engine/Source/Runtime/HeadMountedDisplay/Public/HeadMountedDisplayBase.h
@@ -15,6 +15,10 @@ class HEADMOUNTEDDISPLAY_API FHeadMountedDisplayBase : public FXRTrackingSystemB
 {
 
 public:
+	FHeadMountedDisplayBase()
+		: bHeadTrackingEnforced(false)
+	{}
+
 	virtual ~FHeadMountedDisplayBase() {}
 
 	/**
@@ -27,6 +31,11 @@ public:
 	 */
 	virtual bool IsHeadTrackingAllowed() const override;
 
+	/** Optional IXRTrackingSystem methods.
+	  */
+	virtual bool IsHeadTrackingEnforced() const override;
+	virtual void SetHeadTrackingEnforced(bool bEnabled) override;
+
 	/** 
 	 * Default stereo layer implementation
 	 */
@@ -59,7 +68,7 @@ protected:
 	 */
 	virtual bool PopulateAnalyticsAttributes(TArray<struct FAnalyticsEventAttribute>& EventAttributes);
 
-	/** 
+	/**
 	 * Implement this method to provide an alternate render target for head locked stereo layer rendering, when using the default Stereo Layers implementation.
 	 * 
 	 * Return a FTexture2DRHIRef pointing to a texture that can be composed on top of each eye without applying reprojection to it.
@@ -88,4 +97,7 @@ protected:
 	 */
 	static void CVarSinkHandler();
 	static FAutoConsoleVariableSink CVarSink;
+
+private:
+	bool bHeadTrackingEnforced;
 };
diff --git a/Engine/Source/Runtime/HeadMountedDisplay/Public/IXRTrackingSystem.h b/Engine/Source/Runtime/HeadMountedDisplay/Public/IXRTrackingSystem.h
--- a/Engine/Source/Runtime/HeadMountedDisplay/Public/IXRTrackingSystem.h
+++ b/Engine/Source/Runtime/HeadMountedDisplay/Public/IXRTrackingSystem.h
@@ -299,7 +299,7 @@ public:
 	 */
 	virtual IXRInput* GetXRInput() { return nullptr; }
 
-	/*** XR System related methods moved from IHeadMontedDisplay ***/
+	/*** XR System related methods moved from IHeadMountedDisplay ***/
 
 	/**
 	* Returns true, if head tracking is allowed. Most common case: it returns true when GEngine->IsStereoscopic3D() is true,
@@ -307,6 +307,20 @@ public:
 	*/
 	virtual bool IsHeadTrackingAllowed() const = 0;
 
+	/** 
+	* Can be used to enforce tracking even when stereo rendering is disabled. 
+	* The default implementation does not allow enforcing tracking and always returns false.
+	* This method is called both from the game and render threads.
+	*/
+	virtual bool IsHeadTrackingEnforced() const { return false; }
+
+	/**
+	* Can be used to enforce tracking even when stereo rendering is disabled.
+	* The default implementation does not allow enforcing tracking and ignores the argument.
+	*/
+	virtual void SetHeadTrackingEnforced(bool bEnabled) {};
+
+
 	/**
 	* This method is called when playing begins. Useful to reset all runtime values stored in the plugin.
 	*/
