Fixing clang build errors (non-constant-expression cannot be narrowed from type 'int' to 'uint32' (aka 'unsigned int') in initializer list [-Wc++11-narrowing]).

diff --git a/Engine/Plugins/Runtime/Oculus/OculusVR/Source/OculusMR/Private/OculusMR_CastingCameraActor.cpp b/Engine/Plugins/Runtime/Oculus/OculusVR/Source/OculusMR/Private/OculusMR_CastingCameraActor.cpp
--- a/Engine/Plugins/Runtime/Oculus/OculusVR/Source/OculusMR/Private/OculusMR_CastingCameraActor.cpp
+++ b/Engine/Plugins/Runtime/Oculus/OculusVR/Source/OculusMR/Private/OculusMR_CastingCameraActor.cpp
@@ -472,8 +472,8 @@ void AOculusMR_CastingCameraActor::UpdateCameraColorTexture(const ovrpSizei &fra
 		SrcData,
 		Pitch,
 		(FTexture2DResource*)CameraColorTexture->Resource,
-		frameSize.w,
-		frameSize.h
+		(uint32) frameSize.w,
+		(uint32) frameSize.h
 	};
 
 	ENQUEUE_UNIQUE_RENDER_COMMAND_ONEPARAMETER(
@@ -528,8 +528,8 @@ void AOculusMR_CastingCameraActor::UpdateCameraDepthTexture(const ovrpSizei &fra
 		SrcData,
 		Pitch,
 		(FTexture2DResource*)CameraDepthTexture->Resource,
-		frameSize.w,
-		frameSize.h
+		(uint32) frameSize.w,
+		(uint32) frameSize.h
 	};
 
 	ENQUEUE_UNIQUE_RENDER_COMMAND_ONEPARAMETER(
