Move pointer use inside nullptr check block

diff --git a/Engine/Plugins/Runtime/Oculus/OculusVR/Source/OculusEditor/Private/OculusPlatformToolWidget.cpp b/Engine/Plugins/Runtime/Oculus/OculusVR/Source/OculusEditor/Private/OculusPlatformToolWidget.cpp
--- a/Engine/Plugins/Runtime/Oculus/OculusVR/Source/OculusEditor/Private/OculusPlatformToolWidget.cpp
+++ b/Engine/Plugins/Runtime/Oculus/OculusVR/Source/OculusEditor/Private/OculusPlatformToolWidget.cpp
@@ -458,16 +458,16 @@ FReply SOculusPlatformToolWidget::OnSelectRiftBuildDirectory()
 	TSharedPtr<SWindow> parentWindow = FSlateApplication::Get().FindWidgetWindow(AsShared());
 	const void* parentWindowHandle = (parentWindow.IsValid() && parentWindow->GetNativeWindow().IsValid()) ? parentWindow->GetNativeWindow()->GetOSWindowHandle() : nullptr;
 
-	FString path;
-	FString defaultPath = PlatformSettings->OculusRiftBuildDirectory.IsEmpty() ? FPaths::ProjectContentDir() : PlatformSettings->OculusRiftBuildDirectory;
-	if (FDesktopPlatformModule::Get()->OpenDirectoryDialog(parentWindowHandle, "Choose Rift Build Directory", defaultPath, path))
+	if (PlatformSettings != NULL)
 	{
-		if (PlatformSettings != NULL)
+		FString path;
+		FString defaultPath = PlatformSettings->OculusRiftBuildDirectory.IsEmpty() ? FPaths::ProjectContentDir() : PlatformSettings->OculusRiftBuildDirectory;
+		if (FDesktopPlatformModule::Get()->OpenDirectoryDialog(parentWindowHandle, "Choose Rift Build Directory", defaultPath, path))
 		{
 			PlatformSettings->OculusRiftBuildDirectory = path;
+			PlatformSettings->SaveConfig();
+			BuildGeneralSettingsBox(GeneralSettingsBox);
 		}
-		PlatformSettings->SaveConfig();
-		BuildGeneralSettingsBox(GeneralSettingsBox);
 	}
 	return FReply::Handled();
 }
@@ -488,20 +488,20 @@ FReply SOculusPlatformToolWidget::OnSelectLaunchFilePath()
 	TSharedPtr<SWindow> parentWindow = FSlateApplication::Get().FindWidgetWindow(AsShared());
 	const void* parentWindowHandle = (parentWindow.IsValid() && parentWindow->GetNativeWindow().IsValid()) ? parentWindow->GetNativeWindow()->GetOSWindowHandle() : nullptr;
 
-	TArray<FString> path;
-	FString defaultPath = PlatformSettings->GetLaunchFilePath().IsEmpty() ? FPaths::ProjectContentDir() : PlatformSettings->GetLaunchFilePath();
-	FString fileType = PlatformSettings->GetTargetPlatform() == (uint8)EOculusPlatformTarget::Rift ? "Executables (*.exe)|*.exe" : "APKs (*.apk)|*.apk";
-	if (FDesktopPlatformModule::Get()->OpenFileDialog(parentWindowHandle, "Choose Launch File", defaultPath, defaultPath, fileType, EFileDialogFlags::None, path))
+	if (PlatformSettings != NULL)
 	{
-		if (PlatformSettings != NULL)
+		TArray<FString> path;
+		FString defaultPath = PlatformSettings->GetLaunchFilePath().IsEmpty() ? FPaths::ProjectContentDir() : PlatformSettings->GetLaunchFilePath();
+		FString fileType = PlatformSettings->GetTargetPlatform() == (uint8)EOculusPlatformTarget::Rift ? "Executables (*.exe)|*.exe" : "APKs (*.apk)|*.apk";
+		if (FDesktopPlatformModule::Get()->OpenFileDialog(parentWindowHandle, "Choose Launch File", defaultPath, defaultPath, fileType, EFileDialogFlags::None, path))
 		{
 			if (path.Num() > 0)
 			{
 				PlatformSettings->SetLaunchFilePath(FPaths::ConvertRelativePathToFull(path[0]));
 			}
 			PlatformSettings->SaveConfig();
+			BuildGeneralSettingsBox(GeneralSettingsBox);
 		}
-		BuildGeneralSettingsBox(GeneralSettingsBox);
 	}
 	return FReply::Handled();
 }
