The Oculus HMD plugin provided a wrapper to buffer render targets called "TextureSetProxy". This will be useful to have in other HMD plugins (e.g. SteamVR) since it's pretty general. This change includes some code refactoring to make that class general, and moves the Oculus plugin code over to using the new generalized implementation.

diff --git a/Engine/Plugins/Lumin/MagicLeap/Source/MagicLeapHelperVulkan/Private/MagicLeapHelperVulkan.cpp b/Engine/Plugins/Lumin/MagicLeap/Source/MagicLeapHelperVulkan/Private/MagicLeapHelperVulkan.cpp
--- a/Engine/Plugins/Lumin/MagicLeap/Source/MagicLeapHelperVulkan/Private/MagicLeapHelperVulkan.cpp
+++ b/Engine/Plugins/Lumin/MagicLeap/Source/MagicLeapHelperVulkan/Private/MagicLeapHelperVulkan.cpp
@@ -293,7 +293,6 @@ bool FMagicLeapHelperVulkan::GetMediaTexture(FTextureRHIRef& ResultTexture, FSam
 void FMagicLeapHelperVulkan::AliasMediaTexture(FTextureRHIParamRef DestTexture, FTextureRHIParamRef SrcTexture)
 {
 #if PLATFORM_LUMIN
-	FVulkanDynamicRHI* const RHI = (FVulkanDynamicRHI*)GDynamicRHI;
-	RHI->RHIAliasTextureResources(DestTexture, SrcTexture);
+	GDynamicRHI->RHIAliasTextureResources(DestTexture, SrcTexture);
 #endif
 }
diff --git a/Engine/Plugins/Runtime/Oculus/OculusVR/Source/OculusHMD/Private/OculusHMD.cpp b/Engine/Plugins/Runtime/Oculus/OculusVR/Source/OculusHMD/Private/OculusHMD.cpp
--- a/Engine/Plugins/Runtime/Oculus/OculusVR/Source/OculusHMD/Private/OculusHMD.cpp
+++ b/Engine/Plugins/Runtime/Oculus/OculusVR/Source/OculusHMD/Private/OculusHMD.cpp
@@ -1409,11 +1409,11 @@ namespace OculusHMD
 
 			UE_LOG(LogHMD, Log, TEXT("Allocating Oculus %d x %d rendertarget swapchain"), SizeX, SizeY);
 
-			const FTextureSetProxyPtr& TextureSetProxy = EyeLayer_RenderThread->GetTextureSetProxy();
+			const FXRSwapChainPtr& SwapChain = EyeLayer_RenderThread->GetSwapChain();
 
-			if (TextureSetProxy.IsValid())
+			if (SwapChain.IsValid())
 			{
-				OutTargetableTexture = OutShaderResourceTexture = TextureSetProxy->GetTexture2DArray() ? TextureSetProxy->GetTexture2DArray() : TextureSetProxy->GetTexture2D();
+				OutTargetableTexture = OutShaderResourceTexture = SwapChain->GetTexture2DArray() ? SwapChain->GetTexture2DArray() : SwapChain->GetTexture2D();
 				bNeedReAllocateViewportRenderTarget = false;
 				return true;
 			}
@@ -1431,12 +1431,12 @@ namespace OculusHMD
 
 		if (EyeLayer_RenderThread.IsValid())
 		{
-			const FTextureSetProxyPtr& TextureSet = EyeLayer_RenderThread->GetDepthTextureSetProxy();
+			const FXRSwapChainPtr& SwapChain = EyeLayer_RenderThread->GetDepthSwapChain();
 
-			if (TextureSet.IsValid())
+			if (SwapChain.IsValid())
 			{
 				// Ensure the texture size matches the eye layer. We may get other depth allocations unrelated to the main scene render.
-				if (FIntPoint(SizeX, SizeY) == TextureSet->GetTexture2D()->GetSizeXY())
+				if (FIntPoint(SizeX, SizeY) == SwapChain->GetTexture2D()->GetSizeXY())
 				{
 					if (bNeedReAllocateDepthTexture_RenderThread)
 					{
@@ -1444,8 +1444,8 @@ namespace OculusHMD
 						bNeedReAllocateDepthTexture_RenderThread = false;
 					}
 
-					OutTargetableTexture = TextureSet->GetTexture2D();
-					OutShaderResourceTexture = TextureSet->GetTexture2D();
+					OutTargetableTexture = SwapChain->GetTexture2D();
+					OutShaderResourceTexture = SwapChain->GetTexture2D();
 					return true;
 				}
 			}
@@ -1664,32 +1664,32 @@ namespace OculusHMD
 			return;
 		}
 
-		if (LayerFound && (*LayerFound)->GetTextureSetProxy().IsValid())
+		if (LayerFound && (*LayerFound)->GetSwapChain().IsValid())
 		{
-			bool bRightTexture = (*LayerFound)->GetRightTextureSetProxy().IsValid();
+			bool bRightTexture = (*LayerFound)->GetRightSwapChain().IsValid();
 			switch ((*LayerFound)->GetDesc().ShapeType)
 			{
 			case IStereoLayers::CubemapLayer:
 				if (bRightTexture)
 				{
-					Texture = (*LayerFound)->GetRightTextureSetProxy()->GetTextureCube();
-					LeftTexture = (*LayerFound)->GetTextureSetProxy()->GetTextureCube();
+					Texture = (*LayerFound)->GetRightSwapChain()->GetTextureCube();
+					LeftTexture = (*LayerFound)->GetSwapChain()->GetTextureCube();
 				}
 				else
 				{
-					Texture = LeftTexture = (*LayerFound)->GetTextureSetProxy()->GetTextureCube();
+					Texture = LeftTexture = (*LayerFound)->GetSwapChain()->GetTextureCube();
 				}				break;
 
 			case IStereoLayers::CylinderLayer:
 			case IStereoLayers::QuadLayer:
 				if (bRightTexture)
 				{
-					Texture = (*LayerFound)->GetRightTextureSetProxy()->GetTexture2D();
-					LeftTexture = (*LayerFound)->GetTextureSetProxy()->GetTexture2D();
+					Texture = (*LayerFound)->GetRightSwapChain()->GetTexture2D();
+					LeftTexture = (*LayerFound)->GetSwapChain()->GetTexture2D();
 				}
 				else
 				{
-					Texture = LeftTexture = (*LayerFound)->GetTextureSetProxy()->GetTexture2D();
+					Texture = LeftTexture = (*LayerFound)->GetSwapChain()->GetTexture2D();
 				}
 				break;
 
@@ -2153,9 +2153,9 @@ namespace OculusHMD
 			InitializeEyeLayer_RenderThread(RHICmdList);
 		});
 
-		if (!EyeLayer_RenderThread.IsValid() || !EyeLayer_RenderThread->GetTextureSetProxy().IsValid())
+		if (!EyeLayer_RenderThread.IsValid() || !EyeLayer_RenderThread->GetSwapChain().IsValid())
 		{
-			UE_LOG(LogHMD, Error, TEXT("Failed to create eye layer texture set."));
+			UE_LOG(LogHMD, Error, TEXT("Failed to create eye layer swap chain."));
 			ShutdownSession();
 			return false;
 		}
@@ -2505,9 +2505,9 @@ namespace OculusHMD
 				Layers_RenderThread.Add(EyeLayer);
 			}
 
-			if (EyeLayer->GetDepthTextureSetProxy().IsValid())
+			if (EyeLayer->GetDepthSwapChain().IsValid())
 			{
-				if (!EyeLayer_RenderThread.IsValid() || EyeLayer->GetDepthTextureSetProxy() != EyeLayer_RenderThread->GetDepthTextureSetProxy())
+				if (!EyeLayer_RenderThread.IsValid() || EyeLayer->GetDepthSwapChain() != EyeLayer_RenderThread->GetDepthSwapChain())
 				{
 					bNeedReAllocateDepthTexture_RenderThread = true;
 				}
diff --git a/Engine/Plugins/Runtime/Oculus/OculusVR/Source/OculusHMD/Private/OculusHMD_CustomPresent.cpp b/Engine/Plugins/Runtime/Oculus/OculusVR/Source/OculusHMD/Private/OculusHMD_CustomPresent.cpp
--- a/Engine/Plugins/Runtime/Oculus/OculusVR/Source/OculusHMD/Private/OculusHMD_CustomPresent.cpp
+++ b/Engine/Plugins/Runtime/Oculus/OculusVR/Source/OculusHMD/Private/OculusHMD_CustomPresent.cpp
@@ -276,7 +276,7 @@ int FCustomPresent::GetSystemRecommendedMSAALevel() const
 }
 
 
-FTextureSetProxyPtr FCustomPresent::CreateTextureSetProxy_RenderThread(uint32 InSizeX, uint32 InSizeY, EPixelFormat InFormat, FClearValueBinding InBinding, uint32 InNumMips, uint32 InNumSamples, uint32 InNumSamplesTileMem, ERHIResourceType InResourceType, const TArray<ovrpTextureHandle>& InTextures, uint32 InTexCreateFlags)
+FXRSwapChainPtr FCustomPresent::CreateSwapChain_RenderThread(uint32 InSizeX, uint32 InSizeY, EPixelFormat InFormat, FClearValueBinding InBinding, uint32 InNumMips, uint32 InNumSamples, uint32 InNumSamplesTileMem, ERHIResourceType InResourceType, const TArray<ovrpTextureHandle>& InTextures, uint32 InTexCreateFlags)
 {
 	CheckInRenderThread();
 
@@ -293,7 +293,7 @@ FTextureSetProxyPtr FCustomPresent::CreateTextureSetProxy_RenderThread(uint32 In
 		}
 	}
 
-	return MakeShareable(new FTextureSetProxy(RHITexture, RHITextureSwapChain));
+	return MakeShareable(new FXRSwapChain(RHITexture, RHITextureSwapChain));
 }
 
 
diff --git a/Engine/Plugins/Runtime/Oculus/OculusVR/Source/OculusHMD/Private/OculusHMD_CustomPresent.h b/Engine/Plugins/Runtime/Oculus/OculusVR/Source/OculusHMD/Private/OculusHMD_CustomPresent.h
--- a/Engine/Plugins/Runtime/Oculus/OculusVR/Source/OculusHMD/Private/OculusHMD_CustomPresent.h
+++ b/Engine/Plugins/Runtime/Oculus/OculusVR/Source/OculusHMD/Private/OculusHMD_CustomPresent.h
@@ -6,7 +6,7 @@
 #if OCULUS_HMD_SUPPORTED_PLATFORMS
 #include "OculusHMD_Settings.h"
 #include "OculusHMD_GameFrame.h"
-#include "OculusHMD_TextureSetProxy.h"
+#include "XRSwapChain.h"
 #include "RHI.h"
 #include "RendererInterface.h"
 #include "IStereoLayers.h"
@@ -66,9 +66,8 @@ public:
 	virtual int GetLayerFlags() const { return 0; }
 
 	virtual FTextureRHIRef CreateTexture_RenderThread(uint32 InSizeX, uint32 InSizeY, EPixelFormat InFormat, FClearValueBinding InBinding, uint32 InNumMips, uint32 InNumSamples, uint32 InNumSamplesTileMem, ERHIResourceType InResourceType, ovrpTextureHandle InTexture, uint32 TexCreateFlags) = 0;
-	FTextureSetProxyPtr CreateTextureSetProxy_RenderThread(uint32 InSizeX, uint32 InSizeY, EPixelFormat InFormat, FClearValueBinding InBinding, uint32 InNumMips, uint32 InNumSamples, uint32 InNumSamplesTileMem, ERHIResourceType InResourceType, const TArray<ovrpTextureHandle>& InTextures, uint32 InTexCreateFlags);
+	FXRSwapChainPtr CreateSwapChain_RenderThread(uint32 InSizeX, uint32 InSizeY, EPixelFormat InFormat, FClearValueBinding InBinding, uint32 InNumMips, uint32 InNumSamples, uint32 InNumSamplesTileMem, ERHIResourceType InResourceType, const TArray<ovrpTextureHandle>& InTextures, uint32 InTexCreateFlags);
 	void CopyTexture_RenderThread(FRHICommandListImmediate& RHICmdList, FTextureRHIParamRef DstTexture, FTextureRHIParamRef SrcTexture, FIntRect DstRect = FIntRect(), FIntRect SrcRect = FIntRect(), bool bAlphaPremultiply = false, bool bNoAlphaWrite = false, bool bInvertY = true, bool sRGBSource = false) const;
-	virtual void AliasTextureResources_RHIThread(FTextureRHIParamRef DestTexture, FTextureRHIParamRef SrcTexture) = 0;
 	virtual void SubmitGPUFrameTime(float GPUFrameTime) { }
 
 	bool supportsSRGB() { return bSupportsSRGB; }
diff --git a/Engine/Plugins/Runtime/Oculus/OculusVR/Source/OculusHMD/Private/OculusHMD_CustomPresent_D3D11.cpp b/Engine/Plugins/Runtime/Oculus/OculusVR/Source/OculusHMD/Private/OculusHMD_CustomPresent_D3D11.cpp
--- a/Engine/Plugins/Runtime/Oculus/OculusVR/Source/OculusHMD/Private/OculusHMD_CustomPresent_D3D11.cpp
+++ b/Engine/Plugins/Runtime/Oculus/OculusVR/Source/OculusHMD/Private/OculusHMD_CustomPresent_D3D11.cpp
@@ -26,7 +26,6 @@ public:
 	virtual bool IsUsingCorrectDisplayAdapter() const override;
 	virtual void* GetOvrpDevice() const override;
 	virtual FTextureRHIRef CreateTexture_RenderThread(uint32 InSizeX, uint32 InSizeY, EPixelFormat InFormat, FClearValueBinding InBinding, uint32 InNumMips, uint32 InNumSamples, uint32 InNumSamplesTileMem, ERHIResourceType InResourceType, ovrpTextureHandle InTexture, uint32 InTexCreateFlags) override;
-	virtual void AliasTextureResources_RHIThread(FTextureRHIParamRef DestTexture, FTextureRHIParamRef SrcTexture) override;
 };
 
 
@@ -107,16 +106,6 @@ FTextureRHIRef FD3D11CustomPresent::CreateTexture_RenderThread(uint32 InSizeX, u
 	}
 }
 
-
-void FD3D11CustomPresent::AliasTextureResources_RHIThread(FTextureRHIParamRef DestTexture, FTextureRHIParamRef SrcTexture)
-{
-	CheckInRHIThread();
-
-	FD3D11DynamicRHI* DynamicRHI = static_cast<FD3D11DynamicRHI*>(GDynamicRHI);
-	DynamicRHI->RHIAliasTextureResources(DestTexture, SrcTexture);
-}
-
-
 //-------------------------------------------------------------------------------------------------
 // APIs
 //-------------------------------------------------------------------------------------------------
diff --git a/Engine/Plugins/Runtime/Oculus/OculusVR/Source/OculusHMD/Private/OculusHMD_CustomPresent_D3D12.cpp b/Engine/Plugins/Runtime/Oculus/OculusVR/Source/OculusHMD/Private/OculusHMD_CustomPresent_D3D12.cpp
--- a/Engine/Plugins/Runtime/Oculus/OculusVR/Source/OculusHMD/Private/OculusHMD_CustomPresent_D3D12.cpp
+++ b/Engine/Plugins/Runtime/Oculus/OculusVR/Source/OculusHMD/Private/OculusHMD_CustomPresent_D3D12.cpp
@@ -26,7 +26,6 @@ public:
 	virtual bool IsUsingCorrectDisplayAdapter() const override;
 	virtual void* GetOvrpDevice() const override;
 	virtual FTextureRHIRef CreateTexture_RenderThread(uint32 InSizeX, uint32 InSizeY, EPixelFormat InFormat, FClearValueBinding InBinding, uint32 InNumMips, uint32 InNumSamples, uint32 InNumSamplesTileMem, ERHIResourceType InResourceType, ovrpTextureHandle InTexture, uint32 InTexCreateFlags) override;
-	virtual void AliasTextureResources_RHIThread(FTextureRHIParamRef DestTexture, FTextureRHIParamRef SrcTexture) override;
 };
 
 
@@ -99,15 +98,6 @@ FTextureRHIRef FD3D12CustomPresent::CreateTexture_RenderThread(uint32 InSizeX, u
 }
 
 
-void FD3D12CustomPresent::AliasTextureResources_RHIThread(FTextureRHIParamRef DestTexture, FTextureRHIParamRef SrcTexture)
-{
-	CheckInRHIThread();
-
-	FD3D12DynamicRHI* DynamicRHI = static_cast<FD3D12DynamicRHI*>(GDynamicRHI);
-	DynamicRHI->RHIAliasTextureResources(DestTexture, SrcTexture);
-}
-
-
 //-------------------------------------------------------------------------------------------------
 // APIs
 //-------------------------------------------------------------------------------------------------
diff --git a/Engine/Plugins/Runtime/Oculus/OculusVR/Source/OculusHMD/Private/OculusHMD_CustomPresent_OpenGL.cpp b/Engine/Plugins/Runtime/Oculus/OculusVR/Source/OculusHMD/Private/OculusHMD_CustomPresent_OpenGL.cpp
--- a/Engine/Plugins/Runtime/Oculus/OculusVR/Source/OculusHMD/Private/OculusHMD_CustomPresent_OpenGL.cpp
+++ b/Engine/Plugins/Runtime/Oculus/OculusVR/Source/OculusHMD/Private/OculusHMD_CustomPresent_OpenGL.cpp
@@ -31,7 +31,6 @@ public:
 	// Implementation of FCustomPresent, called by Plugin itself
 	virtual int GetLayerFlags() const override;
 	virtual FTextureRHIRef CreateTexture_RenderThread(uint32 InSizeX, uint32 InSizeY, EPixelFormat InFormat, FClearValueBinding InBinding, uint32 InNumMips, uint32 InNumSamples, uint32 InNumSamplesTileMem, ERHIResourceType InResourceType, ovrpTextureHandle InTexture, uint32 InTexCreateFlags) override;
-	virtual void AliasTextureResources_RHIThread(FTextureRHIParamRef DestTexture, FTextureRHIParamRef SrcTexture) override;
 	virtual void SubmitGPUFrameTime(float GPUFrameTime) override;
 };
 
@@ -74,15 +73,6 @@ FTextureRHIRef FOpenGLCustomPresent::CreateTexture_RenderThread(uint32 InSizeX,
 	}
 }
 
-
-void FOpenGLCustomPresent::AliasTextureResources_RHIThread(FTextureRHIParamRef DestTexture, FTextureRHIParamRef SrcTexture)
-{
-	CheckInRHIThread();
-
-	FOpenGLDynamicRHI* DynamicRHI = static_cast<FOpenGLDynamicRHI*>(GDynamicRHI);
-	DynamicRHI->RHIAliasTextureResources(DestTexture, SrcTexture);
-}
-
 void FOpenGLCustomPresent::SubmitGPUFrameTime(float GPUFrameTime)
 {
 	FOpenGLDynamicRHI* DynamicRHI = static_cast<FOpenGLDynamicRHI*>(GDynamicRHI);
diff --git a/Engine/Plugins/Runtime/Oculus/OculusVR/Source/OculusHMD/Private/OculusHMD_CustomPresent_Vulkan.cpp b/Engine/Plugins/Runtime/Oculus/OculusVR/Source/OculusHMD/Private/OculusHMD_CustomPresent_Vulkan.cpp
--- a/Engine/Plugins/Runtime/Oculus/OculusVR/Source/OculusHMD/Private/OculusHMD_CustomPresent_Vulkan.cpp
+++ b/Engine/Plugins/Runtime/Oculus/OculusVR/Source/OculusHMD/Private/OculusHMD_CustomPresent_Vulkan.cpp
@@ -31,7 +31,6 @@ public:
 	virtual void* GetOvrpDevice() const override;
 	virtual void* GetOvrpCommandQueue() const override;
 	virtual FTextureRHIRef CreateTexture_RenderThread(uint32 InSizeX, uint32 InSizeY, EPixelFormat InFormat, FClearValueBinding InBinding, uint32 InNumMips, uint32 InNumSamples, uint32 InNumSamplesTileMem, ERHIResourceType InResourceType, ovrpTextureHandle InTexture, uint32 InTexCreateFlags) override;
-	virtual void AliasTextureResources_RHIThread(FTextureRHIParamRef DestTexture, FTextureRHIParamRef SrcTexture) override;
 };
 
 
@@ -137,16 +136,6 @@ FTextureRHIRef FVulkanCustomPresent::CreateTexture_RenderThread(uint32 InSizeX,
 	}
 }
 
-
-void FVulkanCustomPresent::AliasTextureResources_RHIThread(FTextureRHIParamRef DestTexture, FTextureRHIParamRef SrcTexture)
-{
-	CheckInRHIThread();
-
-	FVulkanDynamicRHI* DynamicRHI = static_cast<FVulkanDynamicRHI*>(GDynamicRHI);
-	DynamicRHI->RHIAliasTextureResources(DestTexture, SrcTexture);
-}
-
-
 //-------------------------------------------------------------------------------------------------
 // APIs
 //-------------------------------------------------------------------------------------------------
diff --git a/Engine/Plugins/Runtime/Oculus/OculusVR/Source/OculusHMD/Private/OculusHMD_Layer.cpp b/Engine/Plugins/Runtime/Oculus/OculusVR/Source/OculusHMD/Private/OculusHMD_Layer.cpp
--- a/Engine/Plugins/Runtime/Oculus/OculusVR/Source/OculusHMD/Private/OculusHMD_Layer.cpp
+++ b/Engine/Plugins/Runtime/Oculus/OculusVR/Source/OculusHMD/Private/OculusHMD_Layer.cpp
@@ -68,10 +68,10 @@ FLayer::FLayer(const FLayer& Layer) :
 	Desc(Layer.Desc),
 	OvrpLayerId(Layer.OvrpLayerId),
 	OvrpLayer(Layer.OvrpLayer),
-	TextureSetProxy(Layer.TextureSetProxy),
-	DepthTextureSetProxy(Layer.DepthTextureSetProxy),
-	RightTextureSetProxy(Layer.RightTextureSetProxy),
-	RightDepthTextureSetProxy(Layer.RightDepthTextureSetProxy),
+	SwapChain(Layer.SwapChain),
+	DepthSwapChain(Layer.DepthSwapChain),
+	RightSwapChain(Layer.RightSwapChain),
+	RightDepthSwapChain(Layer.RightDepthSwapChain),
 	bUpdateTexture(Layer.bUpdateTexture),
 	bInvertY(Layer.bInvertY),
 	bHasDepth(Layer.bHasDepth),
@@ -431,10 +431,10 @@ void FLayer::Initialize_RenderThread(const FSettings* Settings, FCustomPresent*
 	{
 		OvrpLayerId = InLayer->OvrpLayerId;
 		OvrpLayer = InLayer->OvrpLayer;
-		TextureSetProxy = InLayer->TextureSetProxy;
-		DepthTextureSetProxy = InLayer->DepthTextureSetProxy;
-		RightTextureSetProxy = InLayer->RightTextureSetProxy;
-		RightDepthTextureSetProxy = InLayer->RightDepthTextureSetProxy;
+		SwapChain = InLayer->SwapChain;
+		DepthSwapChain = InLayer->DepthSwapChain;
+		RightSwapChain = InLayer->RightSwapChain;
+		RightDepthSwapChain = InLayer->RightDepthSwapChain;
 		bUpdateTexture = InLayer->bUpdateTexture;
 		bNeedsTexSrgbCreate = InLayer->bNeedsTexSrgbCreate;
 	}
@@ -537,20 +537,20 @@ void FLayer::Initialize_RenderThread(const FSettings* Settings, FCustomPresent*
 			FSceneRenderTargets& SceneContext = FSceneRenderTargets::Get(RHICmdList);
 			FClearValueBinding DepthTextureBinding = SceneContext.GetDefaultDepthClear();
 
-			TextureSetProxy = CustomPresent->CreateTextureSetProxy_RenderThread(SizeX, SizeY, ColorFormat, ColorTextureBinding, NumMips, NumSamples, NumSamplesTileMem, ResourceType, ColorTextures, ColorTexCreateFlags);
+			SwapChain = CustomPresent->CreateSwapChain_RenderThread(SizeX, SizeY, ColorFormat, ColorTextureBinding, NumMips, NumSamples, NumSamplesTileMem, ResourceType, ColorTextures, ColorTexCreateFlags);
 
 			if (bHasDepth)
 			{
-				DepthTextureSetProxy = CustomPresent->CreateTextureSetProxy_RenderThread(SizeX, SizeY, DepthFormat, DepthTextureBinding, 1, NumSamples, NumSamplesTileMem, ResourceType, DepthTextures, DepthTexCreateFlags);
+				DepthSwapChain = CustomPresent->CreateSwapChain_RenderThread(SizeX, SizeY, DepthFormat, DepthTextureBinding, 1, NumSamples, NumSamplesTileMem, ResourceType, DepthTextures, DepthTexCreateFlags);
 			}
 
 			if (OvrpLayerDesc.Layout == ovrpLayout_Stereo)
 			{
-				RightTextureSetProxy = CustomPresent->CreateTextureSetProxy_RenderThread(SizeX, SizeY, ColorFormat, ColorTextureBinding, NumMips, NumSamples, NumSamplesTileMem, ResourceType, RightColorTextures, ColorTexCreateFlags);
+				RightSwapChain = CustomPresent->CreateSwapChain_RenderThread(SizeX, SizeY, ColorFormat, ColorTextureBinding, NumMips, NumSamples, NumSamplesTileMem, ResourceType, RightColorTextures, ColorTexCreateFlags);
 
 				if (bHasDepth)
 				{
-					RightDepthTextureSetProxy = CustomPresent->CreateTextureSetProxy_RenderThread(SizeX, SizeY, DepthFormat, DepthTextureBinding, 1, NumSamples, NumSamplesTileMem, ResourceType, RightDepthTextures, DepthTexCreateFlags);
+					RightDepthSwapChain = CustomPresent->CreateSwapChain_RenderThread(SizeX, SizeY, DepthFormat, DepthTextureBinding, 1, NumSamples, NumSamplesTileMem, ResourceType, RightDepthTextures, DepthTexCreateFlags);
 				}
 			}
 		}
@@ -568,7 +568,7 @@ void FLayer::UpdateTexture_RenderThread(FCustomPresent* CustomPresent, FRHIComma
 {
 	CheckInRenderThread();
 
-	if (bUpdateTexture && TextureSetProxy.IsValid())
+	if (bUpdateTexture && SwapChain.IsValid())
 	{
 		// Copy textures
 		if (Desc.Texture.IsValid())
@@ -580,7 +580,7 @@ void FLayer::UpdateTexture_RenderThread(FCustomPresent* CustomPresent, FRHIComma
 			// Left
 			{
 				FRHITexture* SrcTexture = Desc.LeftTexture.IsValid() ? Desc.LeftTexture : Desc.Texture;
-				FRHITexture* DstTexture = TextureSetProxy->GetTexture();
+				FRHITexture* DstTexture = SwapChain->GetTexture();
 
 				const ovrpRecti& OvrpViewportRect = OvrpLayerSubmit.ViewportRect[ovrpEye_Left];
 				FIntRect DstRect(OvrpViewportRect.Pos.x, OvrpViewportRect.Pos.y, OvrpViewportRect.Pos.x + OvrpViewportRect.Size.w, OvrpViewportRect.Pos.y + OvrpViewportRect.Size.h);
@@ -592,7 +592,7 @@ void FLayer::UpdateTexture_RenderThread(FCustomPresent* CustomPresent, FRHIComma
 			if(OvrpLayerDesc.Layout != ovrpLayout_Mono)
 			{
 				FRHITexture* SrcTexture = Desc.Texture;
-				FRHITexture* DstTexture = RightTextureSetProxy.IsValid() ? RightTextureSetProxy->GetTexture() : TextureSetProxy->GetTexture();
+				FRHITexture* DstTexture = RightSwapChain.IsValid() ? RightSwapChain->GetTexture() : SwapChain->GetTexture();
 
 				const ovrpRecti& OvrpViewportRect = OvrpLayerSubmit.ViewportRect[ovrpEye_Right];
 				FIntRect DstRect(OvrpViewportRect.Pos.x, OvrpViewportRect.Pos.y, OvrpViewportRect.Pos.x + OvrpViewportRect.Size.w, OvrpViewportRect.Pos.y + OvrpViewportRect.Size.h);
@@ -604,11 +604,11 @@ void FLayer::UpdateTexture_RenderThread(FCustomPresent* CustomPresent, FRHIComma
 		}
 
 		// Generate mips
-		TextureSetProxy->GenerateMips_RenderThread(RHICmdList);
+		SwapChain->GenerateMips_RenderThread(RHICmdList);
 
-		if (RightTextureSetProxy.IsValid())
+		if (RightSwapChain.IsValid())
 		{
-			RightTextureSetProxy->GenerateMips_RenderThread(RHICmdList);
+			RightSwapChain->GenerateMips_RenderThread(RHICmdList);
 		}
 	}
 }
@@ -617,7 +617,7 @@ void FLayer::UpdateTexture_RenderThread(FCustomPresent* CustomPresent, FRHIComma
 const ovrpLayerSubmit* FLayer::UpdateLayer_RHIThread(const FSettings* Settings, const FGameFrame* Frame, const int LayerIndex)
 {
 	OvrpLayerSubmit.LayerId = OvrpLayerId;
-	OvrpLayerSubmit.TextureStage = TextureSetProxy.IsValid() ? TextureSetProxy->GetSwapChainIndex_RHIThread() : 0;
+	OvrpLayerSubmit.TextureStage = SwapChain.IsValid() ? SwapChain->GetSwapChainIndex_RHIThread() : 0;
 
 	bool injectColorScale = Id == 0 || Settings->bApplyColorScaleAndOffsetToAllLayers;
 	OvrpLayerSubmit.ColorOffset = injectColorScale ? Settings->ColorOffset : ovrpVector4f{ 0, 0, 0, 0 };
@@ -724,24 +724,24 @@ void FLayer::IncrementSwapChainIndex_RHIThread(FCustomPresent* CustomPresent)
 {
 	CheckInRHIThread();
 
-	if (TextureSetProxy.IsValid())
+	if (SwapChain.IsValid())
 	{
-		TextureSetProxy->IncrementSwapChainIndex_RHIThread(CustomPresent);
+		SwapChain->IncrementSwapChainIndex_RHIThread();
 	}
 
-	if (DepthTextureSetProxy.IsValid())
+	if (DepthSwapChain.IsValid())
 	{
-		DepthTextureSetProxy->IncrementSwapChainIndex_RHIThread(CustomPresent);
+		DepthSwapChain->IncrementSwapChainIndex_RHIThread();
 	}
 
-	if (RightTextureSetProxy.IsValid())
+	if (RightSwapChain.IsValid())
 	{
-		RightTextureSetProxy->IncrementSwapChainIndex_RHIThread(CustomPresent);
+		RightSwapChain->IncrementSwapChainIndex_RHIThread();
 	}
 
-	if (RightDepthTextureSetProxy.IsValid())
+	if (RightDepthSwapChain.IsValid())
 	{
-		RightDepthTextureSetProxy->IncrementSwapChainIndex_RHIThread(CustomPresent);
+		RightDepthSwapChain->IncrementSwapChainIndex_RHIThread();
 	}
 }
 
@@ -752,10 +752,10 @@ void FLayer::ReleaseResources_RHIThread()
 
 	OvrpLayerId = 0;
 	OvrpLayer.Reset();
-	TextureSetProxy.Reset();
-	DepthTextureSetProxy.Reset();
-	RightTextureSetProxy.Reset();
-	RightDepthTextureSetProxy.Reset();
+	SwapChain.Reset();
+	DepthSwapChain.Reset();
+	RightSwapChain.Reset();
+	RightDepthSwapChain.Reset();
 	bUpdateTexture = false;
 }
 
diff --git a/Engine/Plugins/Runtime/Oculus/OculusVR/Source/OculusHMD/Private/OculusHMD_Layer.h b/Engine/Plugins/Runtime/Oculus/OculusVR/Source/OculusHMD/Private/OculusHMD_Layer.h
--- a/Engine/Plugins/Runtime/Oculus/OculusVR/Source/OculusHMD/Private/OculusHMD_Layer.h
+++ b/Engine/Plugins/Runtime/Oculus/OculusVR/Source/OculusHMD/Private/OculusHMD_Layer.h
@@ -6,7 +6,7 @@
 
 #if OCULUS_HMD_SUPPORTED_PLATFORMS
 #include "OculusHMD_CustomPresent.h"
-#include "OculusHMD_TextureSetProxy.h"
+#include "XRSwapChain.h"
 
 namespace OculusHMD
 {
@@ -43,9 +43,9 @@ public:
 	void SetDesc(const IStereoLayers::FLayerDesc& InDesc);
 	const IStereoLayers::FLayerDesc& GetDesc() const { return Desc; }
 	void SetEyeLayerDesc(const ovrpLayerDesc_EyeFov& InEyeLayerDesc, const ovrpRecti InViewportRect[ovrpEye_Count]);
-	const FTextureSetProxyPtr& GetTextureSetProxy() const { return TextureSetProxy; }
-	const FTextureSetProxyPtr& GetRightTextureSetProxy() const { return RightTextureSetProxy; }
-	const FTextureSetProxyPtr& GetDepthTextureSetProxy() const { return DepthTextureSetProxy; }
+	const FXRSwapChainPtr& GetSwapChain() const { return SwapChain; }
+	const FXRSwapChainPtr& GetRightSwapChain() const { return RightSwapChain; }
+	const FXRSwapChainPtr& GetDepthSwapChain() const { return DepthSwapChain; }
 	void MarkTextureForUpdate() { bUpdateTexture = true; }
 #if PLATFORM_ANDROID
 	bool NeedsPokeAHole() { return (Desc.Flags & IStereoLayers::LAYER_FLAG_SUPPORT_DEPTH) != 0; }
@@ -76,10 +76,10 @@ protected:
 	ovrpLayerDescUnion OvrpLayerDesc;
 	ovrpLayerSubmitUnion OvrpLayerSubmit;
 	FOvrpLayerPtr OvrpLayer;
-	FTextureSetProxyPtr TextureSetProxy;
-	FTextureSetProxyPtr DepthTextureSetProxy;
-	FTextureSetProxyPtr RightTextureSetProxy;
-	FTextureSetProxyPtr RightDepthTextureSetProxy;
+	FXRSwapChainPtr SwapChain;
+	FXRSwapChainPtr DepthSwapChain;
+	FXRSwapChainPtr RightSwapChain;
+	FXRSwapChainPtr RightDepthSwapChain;
 	bool bUpdateTexture;
 	bool bInvertY;
 	bool bHasDepth;
diff --git a/Engine/Plugins/Runtime/Oculus/OculusVR/Source/OculusHMD/Private/OculusHMD_TextureSetProxy.cpp b/Engine/Plugins/Runtime/Oculus/OculusVR/Source/OculusHMD/Private/OculusHMD_TextureSetProxy.cpp
deleted file mode 100644
diff --git a/Engine/Plugins/Runtime/Oculus/OculusVR/Source/OculusHMD/Private/OculusHMD_TextureSetProxy.h b/Engine/Plugins/Runtime/Oculus/OculusVR/Source/OculusHMD/Private/OculusHMD_TextureSetProxy.h
deleted file mode 100644
diff --git a/Engine/Plugins/Runtime/OpenXR/Source/OpenXRHMD/Private/OpenXRHMD.cpp b/Engine/Plugins/Runtime/OpenXR/Source/OpenXRHMD/Private/OpenXRHMD.cpp
--- a/Engine/Plugins/Runtime/OpenXR/Source/OpenXRHMD/Private/OpenXRHMD.cpp
+++ b/Engine/Plugins/Runtime/OpenXR/Source/OpenXRHMD/Private/OpenXRHMD.cpp
@@ -253,8 +253,7 @@ void FOpenXRHMD::FOpenXRSwapchain::IncrementSwapChainIndex_RenderThread(XrDurati
 	WaitInfo.timeout = Timeout;
 	XR_ENSURE(xrWaitSwapchainImage(Handle, &WaitInfo));
 
-	FD3D11DynamicRHI* DynamicRHI = static_cast<FD3D11DynamicRHI*>(GDynamicRHI);
-	DynamicRHI->RHIAliasTextureResources(RHITexture, RHITextureSwapChain[SwapChainIndex_RenderThread]);
+	GDynamicRHI->RHIAliasTextureResources(RHITexture, RHITextureSwapChain[SwapChainIndex_RenderThread]);
 }
 
 void FOpenXRHMD::FOpenXRSwapchain::ReleaseSwapChainImage_RenderThread()
diff --git a/Engine/Source/Runtime/D3D12RHI/Private/D3D12RHIPrivate.h b/Engine/Source/Runtime/D3D12RHI/Private/D3D12RHIPrivate.h
--- a/Engine/Source/Runtime/D3D12RHI/Private/D3D12RHIPrivate.h
+++ b/Engine/Source/Runtime/D3D12RHI/Private/D3D12RHIPrivate.h
@@ -348,6 +348,7 @@ public:
 	virtual bool RHIGetRenderQueryResult(FRenderQueryRHIParamRef RenderQuery, uint64& OutResult, bool bWait) final override;
 	virtual uint32 RHIGetViewportNextPresentGPUIndex(FViewportRHIParamRef Viewport) final override;
 	virtual FTexture2DRHIRef RHIGetViewportBackBuffer(FViewportRHIParamRef Viewport) final override;
+	virtual void RHIAliasTextureResources(FTextureRHIParamRef DestTexture, FTextureRHIParamRef SrcTexture) final override;
 	virtual void RHIAdvanceFrameForGetViewportBackBuffer(FViewportRHIParamRef Viewport) final override;
 	virtual void RHIAcquireThreadOwnership() final override;
 	virtual void RHIReleaseThreadOwnership() final override;
@@ -381,7 +382,6 @@ public:
 	virtual ID3D12CommandQueue* RHIGetD3DCommandQueue();
 	virtual FTexture2DRHIRef RHICreateTexture2DFromResource(EPixelFormat Format, uint32 TexCreateFlags, const FClearValueBinding& ClearValueBinding, ID3D12Resource* Resource);
 	virtual FTextureCubeRHIRef RHICreateTextureCubeFromResource(EPixelFormat Format, uint32 TexCreateFlags, const FClearValueBinding& ClearValueBinding, ID3D12Resource* Resource);
-	virtual void RHIAliasTextureResources(FTextureRHIParamRef DestTexture, FTextureRHIParamRef SrcTexture);
 
 	//
 	// The Following functions are the _RenderThread version of the above functions. They allow the RHI to control the thread synchronization for greater efficiency.
diff --git a/Engine/Source/Runtime/HeadMountedDisplay/Private/XRSwapChain.cpp b/Engine/Source/Runtime/HeadMountedDisplay/Private/XRSwapChain.cpp
new file mode 100644
--- /dev/null
+++ b/Engine/Source/Runtime/HeadMountedDisplay/Private/XRSwapChain.cpp
@@ -0,0 +1,113 @@
+// Copyright 1998-2019 Epic Games, Inc. All Rights Reserved.
+
+#include "XRSwapChain.h"
+#include "XRThreadUtils.h"
+#include "CoreGlobals.h"
+
+#ifdef _DEBUG
+#include "CoreMinimal.h"
+#include "HAL/RunnableThread.h"
+#include "RHI.h"
+#endif // _DEBUG
+
+//-------------------------------------------------------------------------------------------------
+// FTextureSetProxy
+//-------------------------------------------------------------------------------------------------
+
+
+FORCEINLINE void CheckInRenderThread()
+{
+#if DO_CHECK
+	check(IsInRenderingThread());
+#endif
+}
+
+bool InRHIOrValidThread()
+{
+	if (GRenderingThread && !GIsRenderingThreadSuspended.Load(EMemoryOrder::Relaxed))
+	{
+		if (GRHIThreadId)
+		{
+			if (FPlatformTLS::GetCurrentThreadId() == GRHIThreadId)
+			{
+				return true;
+			}
+
+			if (FPlatformTLS::GetCurrentThreadId() == GRenderingThread->GetThreadID())
+			{
+				return GetImmediateCommandList_ForRenderCommand().Bypass();
+			}
+
+			return false;
+		}
+		else
+		{
+			return FPlatformTLS::GetCurrentThreadId() == GRenderingThread->GetThreadID();
+		}
+	}
+	else
+	{
+		return IsInGameThread();
+	}
+}
+
+
+FORCEINLINE void CheckInRHIThread()
+{
+#if DO_CHECK
+	check(InRHIOrValidThread());
+#endif
+}
+
+FXRSwapChain::FXRSwapChain(FTextureRHIParamRef InRHITexture, const TArray<FTextureRHIRef>& InRHITextureSwapChain)
+	: RHITexture(InRHITexture)
+	, RHITextureSwapChain(InRHITextureSwapChain)
+	, SwapChainIndex_RHIThread(0)
+{
+}
+
+
+FXRSwapChain::~FXRSwapChain()
+{
+	// @todo: The assertion below fires on exit (even though we're on a valid thread) I'm guessing because everything's getting torn 
+	// down and our thread checking isn't valid. Not sure whether to remove the assert here, or figure out a more correct way to 
+	// perform the check.
+//	CheckInRHIThread();
+
+	ExecuteOnRHIThread([this]()
+	{
+		ReleaseResources_RHIThread();
+	});
+}
+
+
+void FXRSwapChain::GenerateMips_RenderThread(FRHICommandListImmediate& RHICmdList)
+{
+	CheckInRenderThread();
+
+	if (RHITexture->GetNumMips() > 1 && RHITexture->GetTextureCube() == nullptr)
+	{
+#if PLATFORM_WINDOWS
+		RHICmdList.GenerateMips(RHITexture);
+#endif
+	}
+}
+
+
+void FXRSwapChain::IncrementSwapChainIndex_RHIThread()
+{
+	CheckInRHIThread();
+
+	SwapChainIndex_RHIThread = (SwapChainIndex_RHIThread + 1) % GetSwapChainLength();
+	GDynamicRHI->RHIAliasTextureResources(RHITexture, RHITextureSwapChain[SwapChainIndex_RHIThread]);
+}
+
+
+void FXRSwapChain::ReleaseResources_RHIThread()
+{
+	CheckInRHIThread();
+
+	RHITexture = nullptr;
+	RHITextureSwapChain.Empty();
+}
+
diff --git a/Engine/Source/Runtime/HeadMountedDisplay/Public/XRSwapChain.h b/Engine/Source/Runtime/HeadMountedDisplay/Public/XRSwapChain.h
new file mode 100644
--- /dev/null
+++ b/Engine/Source/Runtime/HeadMountedDisplay/Public/XRSwapChain.h
@@ -0,0 +1,38 @@
+// Copyright 1998-2019 Epic Games, Inc. All Rights Reserved.
+
+#pragma once
+#include "CoreMinimal.h"
+#include "RHI.h"
+
+//-------------------------------------------------------------------------------------------------
+// FXRSwapChain
+//-------------------------------------------------------------------------------------------------
+
+class HEADMOUNTEDDISPLAY_API FXRSwapChain : public TSharedFromThis<FXRSwapChain, ESPMode::ThreadSafe>
+{
+public:
+	FXRSwapChain(FTextureRHIParamRef InRHITexture, const TArray<FTextureRHIRef>& InRHITextureSwapChain);
+	virtual ~FXRSwapChain();
+
+	FRHITexture* GetTexture() const { return RHITexture.GetReference(); }
+	FRHITexture2D* GetTexture2D() const { return RHITexture->GetTexture2D(); }
+	FRHITexture2DArray* GetTexture2DArray() const { return RHITexture->GetTexture2DArray(); }
+	FRHITextureCube* GetTextureCube() const { return RHITexture->GetTextureCube(); }
+	uint32 GetSwapChainLength() const { return (uint32)RHITextureSwapChain.Num(); }
+
+	void GenerateMips_RenderThread(FRHICommandListImmediate& RHICmdList);
+	uint32 GetSwapChainIndex_RHIThread() { return SwapChainIndex_RHIThread; }
+
+	virtual void IncrementSwapChainIndex_RHIThread();
+	virtual void ReleaseCurrentImage_RHIThread() {}		
+
+protected:
+	virtual void ReleaseResources_RHIThread();
+
+	FTextureRHIRef RHITexture;
+	TArray<FTextureRHIRef> RHITextureSwapChain;
+	uint32 SwapChainIndex_RHIThread;
+};
+
+typedef TSharedPtr<FXRSwapChain, ESPMode::ThreadSafe> FXRSwapChainPtr;
+
diff --git a/Engine/Source/Runtime/NullDrv/Public/NullRHI.h b/Engine/Source/Runtime/NullDrv/Public/NullRHI.h
--- a/Engine/Source/Runtime/NullDrv/Public/NullRHI.h
+++ b/Engine/Source/Runtime/NullDrv/Public/NullRHI.h
@@ -492,6 +492,10 @@ public:
 	virtual void RHIEndScene() final override
 	{
 
+	}
+	virtual void RHIAliasTextureResources(FTextureRHIParamRef DestTexture, FTextureRHIParamRef SrcTexture) final override
+	{
+
 	}
 	virtual void RHIAdvanceFrameForGetViewportBackBuffer(FViewportRHIParamRef Viewport) final override
 	{
diff --git a/Engine/Source/Runtime/OpenGLDrv/Public/OpenGLDrv.h b/Engine/Source/Runtime/OpenGLDrv/Public/OpenGLDrv.h
--- a/Engine/Source/Runtime/OpenGLDrv/Public/OpenGLDrv.h
+++ b/Engine/Source/Runtime/OpenGLDrv/Public/OpenGLDrv.h
@@ -425,6 +425,7 @@ public:
 	virtual FRenderQueryRHIRef RHICreateRenderQuery(ERenderQueryType QueryType) final override;
 	virtual bool RHIGetRenderQueryResult(FRenderQueryRHIParamRef RenderQuery, uint64& OutResult, bool bWait) final override;
 	virtual FTexture2DRHIRef RHIGetViewportBackBuffer(FViewportRHIParamRef Viewport) final override;
+	virtual void RHIAliasTextureResources(FTextureRHIParamRef DestTexture, FTextureRHIParamRef SrcTexture) final override;
 	virtual void RHIAdvanceFrameForGetViewportBackBuffer(FViewportRHIParamRef Viewport) final override;
 	virtual void RHIAcquireThreadOwnership() final override;
 	virtual void RHIReleaseThreadOwnership() final override;
@@ -532,7 +533,6 @@ public:
 	virtual FTexture2DRHIRef RHICreateTexture2DFromResource(EPixelFormat Format, uint32 SizeX, uint32 SizeY, uint32 NumMips, uint32 NumSamples, uint32 NumSamplesTileMem, const FClearValueBinding& ClearValueBinding, GLuint Resource, uint32 Flags);
 	virtual FTexture2DRHIRef RHICreateTexture2DArrayFromResource(EPixelFormat Format, uint32 SizeX, uint32 SizeY, uint32 ArraySize, uint32 NumMips, uint32 NumSamples, uint32 NumSamplesTileMem, const FClearValueBinding& ClearValueBinding, GLuint Resource, uint32 Flags);
 	virtual FTextureCubeRHIRef RHICreateTextureCubeFromResource(EPixelFormat Format, uint32 Size, bool bArray, uint32 ArraySize, uint32 NumMips, uint32 NumSamples, uint32 NumSamplesTileMem, const FClearValueBinding& ClearValueBinding, GLuint Resource, uint32 Flags);
-	virtual void RHIAliasTextureResources(FTextureRHIParamRef DestTexture, FTextureRHIParamRef SrcTexture);
 
 	// Inline copy
 	virtual void RHICopyToStagingBuffer(FVertexBufferRHIParamRef SourceBufferRHI, FStagingBufferRHIParamRef DestinationStagingBufferRHI, uint32 InOffset, uint32 InNumBytes, FGPUFenceRHIParamRef FenceRHI) final override;
diff --git a/Engine/Source/Runtime/RHI/Public/DynamicRHI.h b/Engine/Source/Runtime/RHI/Public/DynamicRHI.h
--- a/Engine/Source/Runtime/RHI/Public/DynamicRHI.h
+++ b/Engine/Source/Runtime/RHI/Public/DynamicRHI.h
@@ -888,6 +888,11 @@ public:
 		return nullptr;
 	}
 
+	virtual void RHIAliasTextureResources(FTextureRHIParamRef DestTexture, FTextureRHIParamRef SrcTexture)
+	{
+		checkNoEntry();
+	}
+
 	// Only relevant with an RHI thread, this advances the backbuffer for the purpose of GetViewportBackBuffer
 	// FlushType: Thread safe
 	virtual void RHIAdvanceFrameForGetViewportBackBuffer(FViewportRHIParamRef Viewport) = 0;
diff --git a/Engine/Source/Runtime/VulkanRHI/Public/VulkanDynamicRHI.h b/Engine/Source/Runtime/VulkanRHI/Public/VulkanDynamicRHI.h
--- a/Engine/Source/Runtime/VulkanRHI/Public/VulkanDynamicRHI.h
+++ b/Engine/Source/Runtime/VulkanRHI/Public/VulkanDynamicRHI.h
@@ -137,6 +137,7 @@ public:
 	virtual FRenderQueryRHIRef RHICreateRenderQuery(ERenderQueryType QueryType) final override;
 	virtual bool RHIGetRenderQueryResult(FRenderQueryRHIParamRef RenderQuery, uint64& OutResult, bool bWait) final override;
 	virtual FTexture2DRHIRef RHIGetViewportBackBuffer(FViewportRHIParamRef Viewport) final override;
+	virtual void RHIAliasTextureResources(FTextureRHIParamRef DestTexture, FTextureRHIParamRef SrcTexture) final override;
 	virtual void RHIAdvanceFrameForGetViewportBackBuffer(FViewportRHIParamRef Viewport) final override;
 	virtual void RHIAcquireThreadOwnership() final override;
 	virtual void RHIReleaseThreadOwnership() final override;
@@ -356,7 +357,6 @@ public:
 	virtual FTexture2DRHIRef RHICreateTexture2DFromResource(EPixelFormat Format, uint32 SizeX, uint32 SizeY, uint32 NumMips, uint32 NumSamples, VkImage Resource, FSamplerYcbcrConversionInitializer& ConversionInitializer, uint32 Flags);
 	virtual FTexture2DArrayRHIRef RHICreateTexture2DArrayFromResource(EPixelFormat Format, uint32 SizeX, uint32 SizeY, uint32 ArraySize, uint32 NumMips, VkImage Resource, uint32 Flags);
 	virtual FTextureCubeRHIRef RHICreateTextureCubeFromResource(EPixelFormat Format, uint32 Size, bool bArray, uint32 ArraySize, uint32 NumMips, VkImage Resource, uint32 Flags);
-	virtual void RHIAliasTextureResources(FTextureRHIParamRef DestTexture, FTextureRHIParamRef SrcTexture);
 
 	void RHICalibrateTimers() override;
 
diff --git a/Engine/Source/Runtime/Windows/D3D11RHI/Private/D3D11RHIPrivate.h b/Engine/Source/Runtime/Windows/D3D11RHI/Private/D3D11RHIPrivate.h
--- a/Engine/Source/Runtime/Windows/D3D11RHI/Private/D3D11RHIPrivate.h
+++ b/Engine/Source/Runtime/Windows/D3D11RHI/Private/D3D11RHIPrivate.h
@@ -476,6 +476,7 @@ public:
 	virtual FRenderQueryRHIRef RHICreateRenderQuery(ERenderQueryType QueryType) final override;
 	virtual bool RHIGetRenderQueryResult(FRenderQueryRHIParamRef RenderQuery, uint64& OutResult, bool bWait) final override;
 	virtual FTexture2DRHIRef RHIGetViewportBackBuffer(FViewportRHIParamRef Viewport) final override;
+	virtual void RHIAliasTextureResources(FTextureRHIParamRef DestTexture, FTextureRHIParamRef SrcTexture) final override;
 	virtual void RHIAdvanceFrameForGetViewportBackBuffer(FViewportRHIParamRef Viewport) final override;
 	virtual void RHIAcquireThreadOwnership() final override;
 	virtual void RHIReleaseThreadOwnership() final override;
@@ -600,8 +601,7 @@ public:
 
 	virtual FTexture2DRHIRef RHICreateTexture2DFromResource(EPixelFormat Format, uint32 TexCreateFlags, const FClearValueBinding& ClearValueBinding, ID3D11Texture2D* Resource);
 	virtual FTextureCubeRHIRef RHICreateTextureCubeFromResource(EPixelFormat Format, uint32 TexCreateFlags, const FClearValueBinding& ClearValueBinding, ID3D11Texture2D* Resource);
-	virtual void RHIAliasTextureResources(FTextureRHIParamRef DestTexture, FTextureRHIParamRef SrcTexture);
-
+	
 	virtual void RHIPerFrameRHIFlushComplete() final override;
 	virtual void RHIPollRenderQueryResults() final override;
 
