Fix splash delta rotation. Need to update rotation on renderthread copies of splash layers instead of the game thread originals.

diff --git a/Engine/Plugins/Runtime/Oculus/OculusVR/Source/OculusHMD/Private/OculusHMD.cpp b/Engine/Plugins/Runtime/Oculus/OculusVR/Source/OculusHMD/Private/OculusHMD.cpp
--- a/Engine/Plugins/Runtime/Oculus/OculusVR/Source/OculusHMD/Private/OculusHMD.cpp
+++ b/Engine/Plugins/Runtime/Oculus/OculusVR/Source/OculusHMD/Private/OculusHMD.cpp
@@ -533,15 +533,21 @@ namespace OculusHMD
 	{
 		CheckInGameThread();
 
+#if WITH_EDITOR
 		if (GIsEditor && !GEnableVREditorHacks)
 		{
 			// @todo vreditor: If we add support for starting PIE while in VR Editor, we don't want to kill stereo mode when exiting PIE
+			if (Splash->IsShown())
+			{
+				Splash->Hide();
+			}
 			EnableStereo(false);
 			ReleaseDevice();
 
 			FApp::SetUseVRFocus(false);
 			FApp::SetHasVRFocus(false);
 		}
+#endif
 	}
 
 	DECLARE_STATS_GROUP(TEXT("Oculus System Metrics"), STATGROUP_OculusSystemMetrics, STATCAT_Advanced);
diff --git a/Engine/Plugins/Runtime/Oculus/OculusVR/Source/OculusHMD/Private/OculusHMD_Splash.cpp b/Engine/Plugins/Runtime/Oculus/OculusVR/Source/OculusHMD/Private/OculusHMD_Splash.cpp
--- a/Engine/Plugins/Runtime/Oculus/OculusVR/Source/OculusHMD/Private/OculusHMD_Splash.cpp
+++ b/Engine/Plugins/Runtime/Oculus/OculusVR/Source/OculusHMD/Private/OculusHMD_Splash.cpp
@@ -74,24 +74,20 @@ void FSplash::Tick_RenderThread(float DeltaTime)
 	const double TimeInSeconds = FPlatformTime::Seconds();
 	const double DeltaTimeInSeconds = TimeInSeconds - LastTimeInSeconds;
 
-	if (DeltaTimeInSeconds > 2.f * SystemDisplayInterval)
+	if (DeltaTimeInSeconds > 2.f * SystemDisplayInterval && Layers_RenderThread_DeltaRotation.Num() > 0)
 	{
-		for (int32 SplashLayerIndex = 0; SplashLayerIndex < SplashLayers.Num(); SplashLayerIndex++)
+		FScopeLock ScopeLock(&RenderThreadLock);
+		for (TTuple<FLayerPtr, FQuat>& Info : Layers_RenderThread_DeltaRotation)
 		{
-			FSplashLayer& SplashLayer = SplashLayers[SplashLayerIndex];
-
-			if (SplashLayer.Layer.IsValid() && !SplashLayer.Desc.DeltaRotation.Equals(FQuat::Identity))
-			{
-				FScopeLock ScopeLock(&RenderThreadLock);
-
-				IStereoLayers::FLayerDesc LayerDesc = SplashLayer.Layer->GetDesc();
-				LayerDesc.Transform.SetRotation(SplashLayer.Desc.DeltaRotation * LayerDesc.Transform.GetRotation());
-
-				FLayer* Layer = new FLayer(*SplashLayer.Layer);
-				Layer->SetDesc(LayerDesc);
-				SplashLayer.Layer = MakeShareable(Layer);
-
-			}
+			FLayerPtr Layer = Info.Key;
+			const FQuat& DeltaRotation = Info.Value;
+			check(Layer.IsValid());
+			check(!DeltaRotation.Equals(FQuat::Identity)); // Only layers with non-zero delta rotation should be in the DeltaRotation array.
+
+			IStereoLayers::FLayerDesc LayerDesc = Layer->GetDesc();
+			LayerDesc.Transform.SetRotation(LayerDesc.Transform.GetRotation() * DeltaRotation);
+			LayerDesc.Transform.NormalizeRotation();
+			Layer->SetDesc(LayerDesc);
 		}
 		LastTimeInSeconds = TimeInSeconds;
 	}
@@ -154,10 +150,10 @@ void FSplash::StopTicker()
 	{
 		ExecuteOnRenderThread([this]()
 		{
-			if (this->Ticker.IsValid())
+			if (Ticker.IsValid())
 			{
-				this->Ticker->Unregister();
-				this->Ticker = nullptr;
+				Ticker->Unregister();
+				Ticker = nullptr;
 			}
 		});
 		UnloadTextures();
@@ -174,7 +170,8 @@ void FSplash::StartTicker()
 
 		ExecuteOnRenderThread([this]()
 		{
-			this->Ticker->Register();
+			LastTimeInSeconds = FPlatformTime::Seconds();
+			Ticker->Register();
 		});
 	}
 
@@ -464,6 +461,7 @@ void FSplash::Show()
 	{
 		//add oculus-generated layers through the OculusVR settings area
 		FScopeLock ScopeLock(&RenderThreadLock);
+		Layers_RenderThread_DeltaRotation.Reset();
 		Layers_RenderThread_Input.Reset();
 		for (int32 SplashLayerIndex = 0; SplashLayerIndex < SplashLayers.Num(); SplashLayerIndex++)
 		{
@@ -471,7 +469,14 @@ void FSplash::Show()
 
 			if (SplashLayer.Layer.IsValid())
 			{
-				Layers_RenderThread_Input.Add(SplashLayer.Layer->Clone());
+				FLayerPtr ClonedLayer = SplashLayer.Layer->Clone();
+				Layers_RenderThread_Input.Add(ClonedLayer);
+
+				// Register layers that need to be rotated every n ticks
+				if (!SplashLayer.Desc.DeltaRotation.Equals(FQuat::Identity))
+				{
+					Layers_RenderThread_DeltaRotation.Emplace(ClonedLayer, SplashLayer.Desc.DeltaRotation);
+				}
 			}
 		}
 
diff --git a/Engine/Plugins/Runtime/Oculus/OculusVR/Source/OculusHMD/Private/OculusHMD_Splash.h b/Engine/Plugins/Runtime/Oculus/OculusVR/Source/OculusHMD/Private/OculusHMD_Splash.h
--- a/Engine/Plugins/Runtime/Oculus/OculusVR/Source/OculusHMD/Private/OculusHMD_Splash.h
+++ b/Engine/Plugins/Runtime/Oculus/OculusVR/Source/OculusHMD/Private/OculusHMD_Splash.h
@@ -93,6 +93,7 @@ protected:
 	uint32 NextLayerId;
 	FLayerPtr BlackLayer;
 	FLayerPtr UELayer;
+	TArray<TTuple<FLayerPtr, FQuat>> Layers_RenderThread_DeltaRotation;
 	TArray<FLayerPtr> Layers_RenderThread_Input;
 	TArray<FLayerPtr> Layers_RenderThread;
 	TArray<FLayerPtr> Layers_RHIThread;
