Ensure we use the correct texture type when querying for dimensions to avoid a crash.

diff --git a/Engine/Plugins/Runtime/Oculus/OculusVR/Source/OculusHMD/Private/OculusHMD.cpp b/Engine/Plugins/Runtime/Oculus/OculusVR/Source/OculusHMD/Private/OculusHMD.cpp
--- a/Engine/Plugins/Runtime/Oculus/OculusVR/Source/OculusHMD/Private/OculusHMD.cpp
+++ b/Engine/Plugins/Runtime/Oculus/OculusVR/Source/OculusHMD/Private/OculusHMD.cpp
@@ -1391,7 +1391,8 @@ namespace OculusHMD
 			InitializeEyeLayer_RenderThread(GetImmediateCommandList_ForRenderCommand());
 
 			const FXRSwapChainPtr& SwapChain = EyeLayer_RenderThread->GetSwapChain();
-			UE_LOG(LogHMD, Log, TEXT("Allocating Oculus %d x %d rendertarget swapchain"), SwapChain->GetTexture2D()->GetSizeX(), SwapChain->GetTexture2D()->GetSizeY());
+			const FRHITexture2D* const SwapChainTexture = SwapChain->GetTexture2DArray() ? SwapChain->GetTexture2DArray() : SwapChain->GetTexture2D();
+			UE_LOG(LogHMD, Log, TEXT("Allocating Oculus %d x %d rendertarget swapchain"), SwapChainTexture->GetSizeX(), SwapChainTexture->GetSizeY());
 		});
 
 		bNeedReAllocateViewportRenderTarget = true;
