Work around log spam in OculusHMD_Splash.cpp when running PIE with more than one player.
Caused by attempting to call ovrp_WaitToBeginFrame on each PIE window.

diff --git a/Engine/Plugins/Runtime/Oculus/OculusVR/Source/OculusHMD/Private/OculusHMD_Splash.cpp b/Engine/Plugins/Runtime/Oculus/OculusVR/Source/OculusHMD/Private/OculusHMD_Splash.cpp
--- a/Engine/Plugins/Runtime/Oculus/OculusVR/Source/OculusHMD/Private/OculusHMD_Splash.cpp
+++ b/Engine/Plugins/Runtime/Oculus/OculusVR/Source/OculusHMD/Private/OculusHMD_Splash.cpp
@@ -108,7 +108,18 @@ void FSplash::LoadSettings()
 	UStereoLayerFunctionLibrary::EnableAutoLoadingSplashScreen(HMDSettings->bAutoEnabled);
 	if (HMDSettings->bAutoEnabled)
 	{
-		FCoreUObjectDelegates::PreLoadMap.AddSP(this, &FSplash::OnPreLoadMap);
+		if (!LoadLevelDelegate.IsValid())
+		{
+			LoadLevelDelegate = FCoreUObjectDelegates::PreLoadMap.AddSP(this, &FSplash::OnPreLoadMap);
+		}
+	}
+	else
+	{
+		if (LoadLevelDelegate.IsValid())
+		{
+			FCoreUObjectDelegates::PreLoadMap.Remove(LoadLevelDelegate);
+			LoadLevelDelegate.Reset();
+		}
 	}
 }
 
@@ -196,17 +207,24 @@ void FSplash::RenderFrame_RenderThread(FRHICommandListImmediate& RHICmdList)
 	}
 
 	ovrpResult Result;
-	UE_LOG(LogHMD, Verbose, TEXT("Splash ovrp_WaitToBeginFrame %u"), XFrame->FrameNumber);
-	if (OVRP_FAILURE(Result = ovrp_WaitToBeginFrame(XFrame->FrameNumber)))
-	{
-		UE_LOG(LogHMD, Error, TEXT("Splash ovrp_WaitToBeginFrame %u failed (%d)"), XFrame->FrameNumber, Result);
-		XFrame->ShowFlags.Rendering = false;
+	if ( ovrp_GetInitialized() )
+	{ 
+		UE_LOG(LogHMD, Verbose, TEXT("Splash ovrp_WaitToBeginFrame %u"), XFrame->FrameNumber);
+		if (OVRP_FAILURE(Result = ovrp_WaitToBeginFrame(XFrame->FrameNumber)))
+		{
+			UE_LOG(LogHMD, Error, TEXT("Splash ovrp_WaitToBeginFrame %u failed (%d)"), XFrame->FrameNumber, Result);
+			XFrame->ShowFlags.Rendering = false;
+		}
+		else
+		{
+			OculusHMD->WaitFrameNumber = XFrame->FrameNumber;
+			OculusHMD->NextFrameNumber = XFrame->FrameNumber + 1;
+			FPlatformAtomics::InterlockedIncrement(&FramesOutstanding);
+		}
 	}
 	else
 	{
-		OculusHMD->WaitFrameNumber = XFrame->FrameNumber;
-		OculusHMD->NextFrameNumber = XFrame->FrameNumber + 1;
-		FPlatformAtomics::InterlockedIncrement(&FramesOutstanding);
+		XFrame->ShowFlags.Rendering = false;
 	}
 
 	if (XFrame->ShowFlags.Rendering)
diff --git a/Engine/Plugins/Runtime/Oculus/OculusVR/Source/OculusHMD/Private/OculusHMD_Splash.h b/Engine/Plugins/Runtime/Oculus/OculusVR/Source/OculusHMD/Private/OculusHMD_Splash.h
--- a/Engine/Plugins/Runtime/Oculus/OculusVR/Source/OculusHMD/Private/OculusHMD_Splash.h
+++ b/Engine/Plugins/Runtime/Oculus/OculusVR/Source/OculusHMD/Private/OculusHMD_Splash.h
@@ -104,6 +104,7 @@ protected:
 
 	float SystemDisplayInterval;
 	double LastTimeInSeconds;
+	FDelegateHandle LoadLevelDelegate;
 };
 
 typedef TSharedPtr<FSplash> FSplashPtr;
