Fix for bad merge

diff --git a/Engine/Plugins/Runtime/Oculus/OculusVR/Source/OculusHMD/Private/OculusHMD.cpp b/Engine/Plugins/Runtime/Oculus/OculusVR/Source/OculusHMD/Private/OculusHMD.cpp
--- a/Engine/Plugins/Runtime/Oculus/OculusVR/Source/OculusHMD/Private/OculusHMD.cpp
+++ b/Engine/Plugins/Runtime/Oculus/OculusVR/Source/OculusHMD/Private/OculusHMD.cpp
@@ -1390,15 +1390,9 @@ namespace OculusHMD
 		{
 			InitializeEyeLayer_RenderThread(GetImmediateCommandList_ForRenderCommand());
 
->>>> ORIGINAL //UE4/Main/Engine/Plugins/Runtime/Oculus/OculusVR/Source/OculusHMD/Private/OculusHMD.cpp#34
-==== THEIRS //UE4/Main/Engine/Plugins/Runtime/Oculus/OculusVR/Source/OculusHMD/Private/OculusHMD.cpp#35
-			const FTextureSetProxyPtr& TextureSetProxy = EyeLayer_RenderThread->GetTextureSetProxy();
-			UE_LOG(LogHMD, Log, TEXT("Allocating Oculus %d x %d rendertarget swapchain"), TextureSetProxy->GetTexture2D()->GetSizeX(), TextureSetProxy->GetTexture2D()->GetSizeY());
-==== YOURS //Jason.Bestimt_Dev-VR/Engine/Plugins/Runtime/Oculus/OculusVR/Source/OculusHMD/Private/OculusHMD.cpp
 			const FXRSwapChainPtr& SwapChain = EyeLayer_RenderThread->GetSwapChain();
 			const FRHITexture2D* const SwapChainTexture = SwapChain->GetTexture2DArray() ? SwapChain->GetTexture2DArray() : SwapChain->GetTexture2D();
 			UE_LOG(LogHMD, Log, TEXT("Allocating Oculus %d x %d rendertarget swapchain"), SwapChainTexture->GetSizeX(), SwapChainTexture->GetSizeY());
-<<<<
 		});
 
 		bNeedReAllocateViewportRenderTarget = true;
@@ -1432,18 +1426,8 @@ namespace OculusHMD
 
 		if (LayerMap[0].IsValid())
 		{
->>>> ORIGINAL //UE4/Main/Engine/Plugins/Runtime/Oculus/OculusVR/Source/OculusHMD/Private/OculusHMD.cpp#34
-			InitializeEyeLayer_RenderThread(GetImmediateCommandList_ForRenderCommand());
-
-			UE_LOG(LogHMD, Log, TEXT("Allocating Oculus %d x %d rendertarget swapchain"), SizeX, SizeY);
-
-			const FTextureSetProxyPtr& TextureSetProxy = EyeLayer_RenderThread->GetTextureSetProxy();
-==== THEIRS //UE4/Main/Engine/Plugins/Runtime/Oculus/OculusVR/Source/OculusHMD/Private/OculusHMD.cpp#35
-			const FTextureSetProxyPtr& TextureSetProxy = EyeLayer_RenderThread->GetTextureSetProxy();
-==== YOURS //Jason.Bestimt_Dev-VR/Engine/Plugins/Runtime/Oculus/OculusVR/Source/OculusHMD/Private/OculusHMD.cpp
 			const FXRSwapChainPtr& SwapChain = EyeLayer_RenderThread->GetSwapChain();
 			if (SwapChain.IsValid())
-<<<<
 			{
 				OutTargetableTexture = OutShaderResourceTexture = SwapChain->GetTexture2DArray() ? SwapChain->GetTexture2DArray() : SwapChain->GetTexture2D();
 				bNeedReAllocateViewportRenderTarget = false;
