Work around an error on the first frame after re-enabling stereo rendering from blueprints on Oculus Rift.

diff --git a/Engine/Plugins/Runtime/Oculus/OculusVR/Source/OculusHMD/Private/OculusHMD.cpp b/Engine/Plugins/Runtime/Oculus/OculusVR/Source/OculusHMD/Private/OculusHMD.cpp
--- a/Engine/Plugins/Runtime/Oculus/OculusVR/Source/OculusHMD/Private/OculusHMD.cpp
+++ b/Engine/Plugins/Runtime/Oculus/OculusVR/Source/OculusHMD/Private/OculusHMD.cpp
@@ -3038,6 +3038,12 @@ namespace OculusHMD
 			}
 			else
 			{
+				// Work around an error log that can happen when enabling stereo rendering again
+				if (NextFrameNumber == WaitFrameNumber)
+				{
+					NextFrameNumber++;
+				}
+
 				if (Settings->Flags.bPauseRendering)
 				{
 					GEngine->SetMaxFPS(0);
