Fix crash on PIE-exit for D3D12/Vulkan RHIs in Oculus plugin.

diff --git a/Engine/Plugins/Runtime/Oculus/OculusVR/Source/OculusHMD/Private/OculusHMD.cpp b/Engine/Plugins/Runtime/Oculus/OculusVR/Source/OculusHMD/Private/OculusHMD.cpp
--- a/Engine/Plugins/Runtime/Oculus/OculusVR/Source/OculusHMD/Private/OculusHMD.cpp
+++ b/Engine/Plugins/Runtime/Oculus/OculusVR/Source/OculusHMD/Private/OculusHMD.cpp
@@ -720,6 +720,12 @@ namespace OculusHMD
 		FStressTester::TickCPU_GameThread(this);
 #endif
 
+		if (bShutdownRequestQueued)
+		{
+			bShutdownRequestQueued = false;
+			DoSessionShutdown();
+		}
+
 		if (!InWorldContext.World() || (!(GEnableVREditorHacks && InWorldContext.WorldType == EWorldType::Editor) && !InWorldContext.World()->IsGameWorld()))	// @todo vreditor: (Also see OnEndGameFrame()) Kind of a hack here so we can use VR in editor viewports.  We need to consider when running GameWorld viewports inside the editor with VR.
 		{
 			// ignore all non-game worlds
@@ -937,6 +943,62 @@ namespace OculusHMD
 		return retval;
 	}
 
+	void FOculusHMD::DoSessionShutdown()
+	{
+		// Release resources
+		ExecuteOnRenderThread([this]()
+		{
+			ExecuteOnRHIThread([this]()
+			{
+				for (int32 LayerIndex = 0; LayerIndex < Layers_RenderThread.Num(); LayerIndex++)
+				{
+					Layers_RenderThread[LayerIndex]->ReleaseResources_RHIThread();
+				}
+
+				for (int32 LayerIndex = 0; LayerIndex < Layers_RHIThread.Num(); LayerIndex++)
+				{
+					Layers_RHIThread[LayerIndex]->ReleaseResources_RHIThread();
+				}
+
+				if (Splash.IsValid())
+				{
+					Splash->ReleaseResources_RHIThread();
+				}
+
+				if (CustomPresent)
+				{
+					CustomPresent->ReleaseResources_RHIThread();
+				}
+
+				Settings_RHIThread.Reset();
+				Frame_RHIThread.Reset();
+				Layers_RHIThread.Reset();
+			});
+
+			Settings_RenderThread.Reset();
+			Frame_RenderThread.Reset();
+			Layers_RenderThread.Reset();
+			EyeLayer_RenderThread.Reset();
+		});
+
+		Frame.Reset();
+		NextFrameToRender.Reset();
+		LastFrameToRender.Reset();
+
+#if !UE_BUILD_SHIPPING
+		UDebugDrawService::Unregister(DrawDebugDelegateHandle);
+#endif
+
+		// The Editor may release VR focus in OnEndPlay
+		if (!GIsEditor)
+		{
+			FApp::SetUseVRFocus(false);
+			FApp::SetHasVRFocus(false);
+		}
+
+		ShutdownSession();
+	
+	}
 
 	bool FOculusHMD::OnEndGameFrame(FWorldContext& InWorldContext)
 	{
@@ -1871,6 +1933,7 @@ namespace OculusHMD
 		: FHeadMountedDisplayBase(nullptr)
 		, FSceneViewExtensionBase(AutoRegister)
 		, ConsoleCommands(this)
+		, bShutdownRequestQueued(false)
 	{
 		Flags.Raw = 0;
 		OCFlags.Raw = 0;
@@ -2208,59 +2271,8 @@ namespace OculusHMD
 
 		if (ovrp_GetInitialized())
 		{
-
-			// Release resources
-			ExecuteOnRenderThread([this]()
-			{
-				ExecuteOnRHIThread([this]()
-				{
-					for (int32 LayerIndex = 0; LayerIndex < Layers_RenderThread.Num(); LayerIndex++)
-					{
-						Layers_RenderThread[LayerIndex]->ReleaseResources_RHIThread();
-					}
-
-					for (int32 LayerIndex = 0; LayerIndex < Layers_RHIThread.Num(); LayerIndex++)
-					{
-						Layers_RHIThread[LayerIndex]->ReleaseResources_RHIThread();
-					}
-
-					if (Splash.IsValid())
-					{
-						Splash->ReleaseResources_RHIThread();
-					}
-
-					if (CustomPresent)
-					{
-						CustomPresent->ReleaseResources_RHIThread();
-					}
-
-					Settings_RHIThread.Reset();
-					Frame_RHIThread.Reset();
-					Layers_RHIThread.Reset();
-				});
-
-				Settings_RenderThread.Reset();
-				Frame_RenderThread.Reset();
-				Layers_RenderThread.Reset();
-				EyeLayer_RenderThread.Reset();
-			});
-
-			Frame.Reset();
-			NextFrameToRender.Reset();
-			LastFrameToRender.Reset();
-
-#if !UE_BUILD_SHIPPING
-			UDebugDrawService::Unregister(DrawDebugDelegateHandle);
-#endif
-
-			// The Editor may release VR focus in OnEndPlay
-			if (!GIsEditor)
-			{
-				FApp::SetUseVRFocus(false);
-				FApp::SetHasVRFocus(false);
-			}
-
-			ShutdownSession();
+			// Wait until the next frame before ending the session (workaround for DX12/Vulkan resources being ripped out from under us before we're done with them).
+			bShutdownRequestQueued = true;
 		}
 	}
 
diff --git a/Engine/Plugins/Runtime/Oculus/OculusVR/Source/OculusHMD/Private/OculusHMD.h b/Engine/Plugins/Runtime/Oculus/OculusVR/Source/OculusHMD/Private/OculusHMD.h
--- a/Engine/Plugins/Runtime/Oculus/OculusVR/Source/OculusHMD/Private/OculusHMD.h
+++ b/Engine/Plugins/Runtime/Oculus/OculusVR/Source/OculusHMD/Private/OculusHMD.h
@@ -351,6 +351,7 @@ protected:
 #endif
 
 	void LoadFromSettings();
+	void DoSessionShutdown();
 
 protected:
 	void UpdateHMDWornState();
@@ -436,6 +437,8 @@ protected:
 #if !UE_BUILD_SHIPPING
 	FDelegateHandle DrawDebugDelegateHandle;
 #endif
+
+	bool bShutdownRequestQueued;
 };
 
 typedef TSharedPtr< FOculusHMD, ESPMode::ThreadSafe > FOculusHMDPtr;
diff --git a/Engine/Plugins/Runtime/Oculus/OculusVR/Source/OculusHMD/Private/OculusHMD_Layer.cpp b/Engine/Plugins/Runtime/Oculus/OculusVR/Source/OculusHMD/Private/OculusHMD_Layer.cpp
--- a/Engine/Plugins/Runtime/Oculus/OculusVR/Source/OculusHMD/Private/OculusHMD_Layer.cpp
+++ b/Engine/Plugins/Runtime/Oculus/OculusVR/Source/OculusHMD/Private/OculusHMD_Layer.cpp
@@ -34,11 +34,13 @@ FOvrpLayer::FOvrpLayer(uint32 InOvrpLayerId) :
 
 FOvrpLayer::~FOvrpLayer()
 {
-	check(InRenderThread() || InRHIThread());
-	ExecuteOnRHIThread_DoNotWait([this]()
+	if (!IsInGameThread())
 	{
-		ovrp_DestroyLayer(OvrpLayerId);
-	});
+		ExecuteOnRHIThread_DoNotWait([this]()
+		{
+			ovrp_DestroyLayer(OvrpLayerId);
+		});
+	}
 }
 
 
