Integration from parner-oculus-staging
Launcher services, build analytics (telemetry) including new opt-in UI at our request.

diff --git a/Engine/Plugins/Runtime/Oculus/OculusVR/Source/OculusEditor/OculusEditor.Build.cs b/Engine/Plugins/Runtime/Oculus/OculusVR/Source/OculusEditor/OculusEditor.Build.cs
--- a/Engine/Plugins/Runtime/Oculus/OculusVR/Source/OculusEditor/OculusEditor.Build.cs
+++ b/Engine/Plugins/Runtime/Oculus/OculusVR/Source/OculusEditor/OculusEditor.Build.cs
@@ -24,7 +24,9 @@ public class OculusEditor : ModuleRules
 				"OVRPlugin",
 				"HTTP",
 				"DesktopPlatform",
-			}
+                "LauncherServices",
+                "GameProjectGeneration",
+            }
 			);
 
 		PrivateIncludePaths.AddRange(
diff --git a/Engine/Plugins/Runtime/Oculus/OculusVR/Source/OculusEditor/Private/OculusBuildAnalytics.cpp b/Engine/Plugins/Runtime/Oculus/OculusVR/Source/OculusEditor/Private/OculusBuildAnalytics.cpp
new file mode 100644
--- /dev/null
+++ b/Engine/Plugins/Runtime/Oculus/OculusVR/Source/OculusEditor/Private/OculusBuildAnalytics.cpp
@@ -0,0 +1,313 @@
+#include "OculusBuildAnalytics.h"
+#include "GameProjectGenerationModule.h"
+#include "OculusHMDModule.h"
+#include "Runtime/Core/Public/HAL/FileManager.h"
+
+FOculusBuildAnalytics* FOculusBuildAnalytics::instance = 0;
+
+FOculusBuildAnalytics* FOculusBuildAnalytics::GetInstance()
+{
+	if (IOculusHMDModule::IsAvailable())
+	{
+		if (FOculusHMDModule::Get().PreInit())
+		{
+			if (instance == 0)
+			{
+				instance = new FOculusBuildAnalytics();
+			}
+		}
+	}
+
+	return instance;
+}
+
+void FOculusBuildAnalytics::Shutdown()
+{
+	if (instance != NULL)
+	{
+		ovrp_SetDeveloperMode(false);
+	}
+}
+
+FOculusBuildAnalytics::FOculusBuildAnalytics()
+{
+	bool TelemetryEnabled = false;
+	if (!GConfig->GetBool(TEXT("/Script/OculusEditor.OculusEditorSettings"), TEXT("bEnableOculusBuildTelemetry"), TelemetryEnabled, GEditorIni))
+	{
+		GConfig->SetBool(TEXT("/Script/OculusEditor.OculusEditorSettings"), TEXT("bEnableOculusBuildTelemetry"), TelemetryEnabled, GEditorIni);
+		GConfig->Flush(0);
+	}
+
+	if (TelemetryEnabled)
+	{
+		RegisterLauncherCallback();
+	}
+}
+
+void FOculusBuildAnalytics::OnTelemetryToggled(bool Enabled)
+{
+	if (Enabled)
+	{
+		RegisterLauncherCallback();
+	}
+	else
+	{
+		if (LauncherCallbackHandle.IsValid())
+		{
+			ILauncherServicesModule& ProjectLauncherServicesModule = FModuleManager::LoadModuleChecked<ILauncherServicesModule>("LauncherServices");
+			ProjectLauncherServicesModule.OnCreateLauncherDelegate.Remove(LauncherCallbackHandle);
+		}
+	}
+}
+
+void FOculusBuildAnalytics::RegisterLauncherCallback()
+{
+	ILauncherServicesModule& ProjectLauncherServicesModule = FModuleManager::LoadModuleChecked<ILauncherServicesModule>("LauncherServices");
+	LauncherCallbackHandle = ProjectLauncherServicesModule.OnCreateLauncherDelegate.AddRaw(this, &FOculusBuildAnalytics::OnLauncherCreated);
+}
+
+void FOculusBuildAnalytics::OnLauncherCreated(ILauncherRef Launcher)
+{
+	// Add callback for when launcher worker is started
+	Launcher->FLauncherWorkerStartedDelegate.AddRaw(this, &FOculusBuildAnalytics::OnLauncherWorkerStarted);
+}
+
+void FOculusBuildAnalytics::OnLauncherWorkerStarted(ILauncherWorkerPtr LauncherWorker, ILauncherProfileRef Profile)
+{
+	TArray<FString> Platforms = Profile.Get().GetCookedPlatforms();
+	if (Platforms.Num() == 1)
+	{
+		if (Platforms[0].Equals("Android_ASTC") || Platforms[0].Contains("Windows"))
+		{
+			CurrentBuildStage = UNDEFINED_STAGE;
+			AndroidPackageTime = 0;
+			UATLaunched = false;
+			BuildCompleted = false;
+			CurrentBuildPlatform = Platforms[0];
+			TotalBuildTime = 0;
+			BuildStepCount = 0;
+
+			ovrp_SetDeveloperMode(true);
+
+			OutputDirectory = Profile.Get().GetPackageDirectory();
+
+			// Assign callbacks for stages
+			LauncherWorker.Get()->OnStageCompleted().AddRaw(this, &FOculusBuildAnalytics::OnStageCompleted);
+			LauncherWorker.Get()->OnOutputReceived().AddRaw(this, &FOculusBuildAnalytics::OnBuildOutputRecieved);
+			LauncherWorker.Get()->OnStageStarted().AddRaw(this, &FOculusBuildAnalytics::OnStageStarted);
+			LauncherWorker.Get()->OnCompleted().AddRaw(this, &FOculusBuildAnalytics::OnCompleted);
+
+			// Get information on what oculus platform we are building for and also the OS platform
+			FString OculusPlatform;
+			if (CurrentBuildPlatform.Equals("Android_ASTC"))
+			{
+				UEnum* OculusMobileDevices = StaticEnum<EOculusMobileDevice::Type>();
+				UAndroidRuntimeSettings* Settings = GetMutableDefault<UAndroidRuntimeSettings>();
+				TArray<TEnumAsByte<EOculusMobileDevice::Type>> TargetOculusDevices = Settings->PackageForOculusMobile;
+				TArray<FString> Devices;
+
+				if (TargetOculusDevices.Contains(EOculusMobileDevice::GearGo))
+				{
+					Devices.Add("geargo");
+				}
+				if (TargetOculusDevices.Contains(EOculusMobileDevice::Quest))
+				{
+					Devices.Add("quest");
+				}
+				OculusPlatform = FString::Join(Devices, TEXT("_"));
+			}
+			else if (CurrentBuildPlatform.Contains("Windows"))
+			{
+				CurrentBuildPlatform = "Windows";
+				OculusPlatform = "rift";
+			}
+
+			// Count user asset files
+			UserAssetCount = 0;
+			TArray<FString> FileNames;
+			IFileManager::Get().FindFilesRecursive(FileNames, *FPaths::ProjectContentDir(), TEXT("*.*"), true, false, false);
+			UserAssetCount = FileNames.Num();
+
+			// Count user script files
+			FGameProjectGenerationModule& GameProjectModule = FModuleManager::LoadModuleChecked<FGameProjectGenerationModule>(TEXT("GameProjectGeneration"));
+			SourceFileCount = 0;
+			SourceFileDirectorySize = 0;
+			GameProjectModule.Get().GetProjectSourceDirectoryInfo(SourceFileCount, SourceFileDirectorySize);
+
+			// Generate build GUID
+			FGuid guid = FGuid::NewGuid();
+			ovrp_AddCustomMetadata("build_guid", TCHAR_TO_ANSI(*guid.ToString()));
+
+			// Send build start event with corresponding metadata
+			ovrp_AddCustomMetadata("asset_count", TCHAR_TO_ANSI(*FString::FromInt(UserAssetCount)));
+			ovrp_AddCustomMetadata("script_count", TCHAR_TO_ANSI(*FString::FromInt(SourceFileCount)));
+
+			ovrp_AddCustomMetadata("target_platform", TCHAR_TO_ANSI(*CurrentBuildPlatform));
+			ovrp_AddCustomMetadata("target_oculus_platform", TCHAR_TO_ANSI(*OculusPlatform));
+
+			TArray<ILauncherTaskPtr> TaskList;
+			LauncherWorker->GetTasks(TaskList);
+			BuildStepCount = TaskList.Num();
+		}
+	}
+}
+
+void FOculusBuildAnalytics::OnCompleted(bool Succeeded, double TotalTime, int32 ErrorCode)
+{
+	if (!BuildCompleted && Succeeded)
+	{
+		SendBuildCompleteEvent(TotalTime);
+	}
+}
+
+void FOculusBuildAnalytics::OnStageCompleted(const FString& StageName, double Time)
+{
+	if (CurrentBuildStage != UNDEFINED_STAGE)
+	{
+		FString TaskName;
+		switch (CurrentBuildStage)
+		{
+		case COOK_IN_EDITOR_STAGE:	TaskName = "build_step_editor_cook";	break;
+		case LAUNCH_UAT_STAGE:		TaskName = "build_step_launch_uat";		break;
+		case COMPILE_STAGE:			TaskName = "build_step_compile";		break;
+		case COOK_STAGE:			TaskName = "build_step_cook";			break;
+		case DEPLOY_STAGE:			TaskName = "build_step_deploy";			break;
+		case PACKAGE_STAGE:			TaskName = "build_step_package";		break;
+		case RUN_STAGE:				return;
+		default:					TaskName = "build_step_undefined";		break;
+		}
+
+		if (AndroidPackageTime > 0)
+		{
+			Time -= AndroidPackageTime;
+		}
+
+		TotalBuildTime += Time;
+		ovrp_SendEvent2(TCHAR_TO_ANSI(*TaskName), TCHAR_TO_ANSI(*FString::SanitizeFloat(Time)), "ovrbuild");
+	}
+}
+
+void FOculusBuildAnalytics::OnStageStarted(const FString& StageName)
+{
+	if (StageName.Equals("Cooking in the editor"))
+	{
+		CurrentBuildStage = COOK_IN_EDITOR_STAGE;
+	}
+	else if (StageName.Equals("Build Task") && CurrentBuildStage == LAUNCH_UAT_STAGE)
+	{
+		CurrentBuildStage = COMPILE_STAGE;
+	}
+	else if (StageName.Equals("Build Task"))
+	{
+		CurrentBuildStage = LAUNCH_UAT_STAGE;
+	}
+	else if (StageName.Equals("Cook Task"))
+	{
+		CurrentBuildStage = COOK_STAGE;
+	}
+	else if (StageName.Equals("Package Task"))
+	{
+		CurrentBuildStage = PACKAGE_STAGE;
+	}
+	else if (StageName.Equals("Deploy Task"))
+	{
+		CurrentBuildStage = DEPLOY_STAGE;
+	}
+	else if (StageName.Equals("Run Task"))
+	{
+		CurrentBuildStage = RUN_STAGE;
+		SendBuildCompleteEvent(TotalBuildTime);
+		BuildCompleted = true;
+	}
+	else
+	{
+		CurrentBuildStage = UNDEFINED_STAGE;
+	}
+}
+
+void FOculusBuildAnalytics::OnBuildOutputRecieved(const FString& Message)
+{
+	if (CurrentBuildPlatform.Equals("Android_ASTC") && (CurrentBuildStage == DEPLOY_STAGE || CurrentBuildStage == PACKAGE_STAGE))
+	{
+		if (Message.Contains("BUILD SUCCESSFUL"))
+		{
+			FString Text, Time;
+			Message.Split("in", &Text, &Time);
+
+			if (!Time.IsEmpty())
+			{
+				FString SMinutes, SSeconds;
+				if (Time.Contains("m"))
+				{
+					Time.Split("m", &SMinutes, &SSeconds);
+				}
+				else
+				{
+					SSeconds = Time;
+				}
+
+				int Minutes = FCString::Atoi(*SMinutes);
+				int Seconds = FCString::Atoi(*SSeconds);
+
+				AndroidPackageTime = Minutes * 60 + Seconds;
+
+				ovrp_SendEvent2("build_step_gradle_build", TCHAR_TO_ANSI(*FString::SanitizeFloat(AndroidPackageTime)), "ovrbuild");
+			}
+		}
+	}
+}
+
+void FOculusBuildAnalytics::SendBuildCompleteEvent(float TotalTime)
+{
+	if (CurrentBuildPlatform.Equals("Android_ASTC"))
+	{
+		int64 APKTotalSize = 0;
+		TArray<FString> FoundAPKs;
+		OutputDirectory = FPaths::ProjectDir() + "Binaries/Android";
+		OutputDirectory = FPaths::ConvertRelativePathToFull(OutputDirectory);
+		IFileManager::Get().FindFiles(FoundAPKs, *FPaths::Combine(OutputDirectory, TEXT("*.apk")), true, false);
+
+		FDateTime LatestTime = FDateTime(0);
+		FString LatestAPK;
+		for (int i = 0; i < FoundAPKs.Num(); i++)
+		{
+			FDateTime APKCreationTime = IFileManager::Get().GetTimeStamp(*FPaths::Combine(OutputDirectory, FoundAPKs[i]));
+			if (APKCreationTime > LatestTime)
+			{
+				LatestTime = APKCreationTime;
+				LatestAPK = FoundAPKs[i];
+			}
+		}
+
+		TArray<FString> FoundOBBs;
+		LatestTime = FDateTime(0);
+		FString LatestOBB;
+		IFileManager::Get().FindFiles(FoundOBBs, *FPaths::Combine(OutputDirectory, TEXT("*.obb")), true, false);
+		for (int i = 0; i < FoundOBBs.Num(); i++)
+		{
+			FDateTime OBBCreationTime = IFileManager::Get().GetTimeStamp(*FPaths::Combine(OutputDirectory, FoundOBBs[i]));
+			if (OBBCreationTime > LatestTime)
+			{
+				LatestTime = OBBCreationTime;
+				LatestOBB = FoundOBBs[i];
+			}
+		}
+
+		if (!LatestAPK.IsEmpty())
+		{
+			APKTotalSize += IFileManager::Get().FileSize(*FPaths::Combine(OutputDirectory, LatestAPK));
+		}
+		if (!LatestOBB.IsEmpty())
+		{
+			APKTotalSize += IFileManager::Get().FileSize(*FPaths::Combine(OutputDirectory, LatestOBB));
+		}
+
+		if (APKTotalSize > 0)
+		{
+			ovrp_AddCustomMetadata("build_output_size", TCHAR_TO_ANSI(*FString::FromInt(APKTotalSize)));
+		}
+	}
+
+	ovrp_AddCustomMetadata("build_step_count", TCHAR_TO_ANSI(*FString::FromInt(BuildStepCount)));
+	ovrp_SendEvent2("build_complete", TCHAR_TO_ANSI(*FString::SanitizeFloat(TotalTime)), "ovrbuild");
+}
diff --git a/Engine/Plugins/Runtime/Oculus/OculusVR/Source/OculusEditor/Private/OculusBuildAnalytics.h b/Engine/Plugins/Runtime/Oculus/OculusVR/Source/OculusEditor/Private/OculusBuildAnalytics.h
new file mode 100644
--- /dev/null
+++ b/Engine/Plugins/Runtime/Oculus/OculusVR/Source/OculusEditor/Private/OculusBuildAnalytics.h
@@ -0,0 +1,58 @@
+#pragma once
+
+#include "ILauncherServicesModule.h"
+#include "ILauncher.h"
+#include "Modules/ModuleManager.h"
+#include "UObject/Class.h"
+#include "AndroidRuntimeSettings.h"
+#include "OVR_Plugin.h"
+
+enum EBuildStage
+{
+	UNDEFINED_STAGE,
+	COOK_IN_EDITOR_STAGE,
+	COOK_STAGE,
+	LAUNCH_UAT_STAGE,
+	COMPILE_STAGE,
+	PACKAGE_STAGE,
+	DEPLOY_STAGE,
+	RUN_STAGE,
+};
+
+class FOculusBuildAnalytics
+{
+public:
+	static FOculusBuildAnalytics* GetInstance();
+	static void Shutdown();
+
+	void RegisterLauncherCallback();
+	void OnTelemetryToggled(bool Enabled);
+
+	void OnLauncherCreated(ILauncherRef Launcher);
+	void OnLauncherWorkerStarted(ILauncherWorkerPtr LauncherWorker, ILauncherProfileRef Profile);
+	void OnStageCompleted(const FString& StageName, double Time);
+	void OnStageStarted(const FString& StageName);
+	void OnBuildOutputRecieved(const FString& Message);
+	void OnCompleted(bool Succeeded, double TotalTime, int32 ErrorCode);
+	void SendBuildCompleteEvent(float TotalTime);
+
+private:
+	FOculusBuildAnalytics();
+
+	static FOculusBuildAnalytics* instance;
+
+	FDelegateHandle LauncherCallbackHandle;
+
+	float TotalBuildTime;
+	float AndroidPackageTime;
+	bool BuildCompleted;
+	bool UATLaunched;
+	int UserAssetCount;
+	int BuildStepCount;
+	int32 SourceFileCount;
+	int64 SourceFileDirectorySize;
+
+	EBuildStage CurrentBuildStage;
+	FString CurrentBuildPlatform;
+	FString OutputDirectory;
+};
diff --git a/Engine/Plugins/Runtime/Oculus/OculusVR/Source/OculusEditor/Private/OculusEditorModule.cpp b/Engine/Plugins/Runtime/Oculus/OculusVR/Source/OculusEditor/Private/OculusEditorModule.cpp
--- a/Engine/Plugins/Runtime/Oculus/OculusVR/Source/OculusEditor/Private/OculusEditorModule.cpp
+++ b/Engine/Plugins/Runtime/Oculus/OculusVR/Source/OculusEditor/Private/OculusEditorModule.cpp
@@ -46,7 +46,7 @@ void FOculusEditorModule::StartupModule()
 	RegisterSettings();
 	FOculusAssetDirectory::LoadForCook();
 
-	if(!IsRunningCommandlet())
+	if (!IsRunningCommandlet())
 	{
 		FOculusToolStyle::Initialize();
 		FOculusToolStyle::ReloadTextures();
@@ -66,7 +66,7 @@ void FOculusEditorModule::StartupModule()
 		TSharedPtr<FExtender> MenuExtender = MakeShareable(new FExtender());
 		MenuExtender->AddMenuExtension("Miscellaneous", EExtensionHook::After, PluginCommands, FMenuExtensionDelegate::CreateRaw(this, &FOculusEditorModule::AddMenuExtension));
 		LevelEditorModule.GetMenuExtensibilityManager()->AddExtender(MenuExtender);
-		
+
 		/*
 		// If you want to make the tool even easier to launch, and add a toolbar button.
 		TSharedPtr<FExtender> ToolbarExtender = MakeShareable(new FExtender);
@@ -81,7 +81,9 @@ void FOculusEditorModule::StartupModule()
 		FGlobalTabmanager::Get()->RegisterNomadTabSpawner(OculusPlatToolTabName, FOnSpawnTab::CreateRaw(this, &FOculusEditorModule::OnSpawnPlatToolTab))
 			.SetDisplayName(LOCTEXT("FOculusPlatfToolTabTitle", "Oculus Platform Tool"))
 			.SetMenuType(ETabSpawnerMenuType::Hidden);
-	 }
+
+		FCoreDelegates::OnFEngineLoopInitComplete.AddRaw(this, &FOculusEditorModule::OnEngineLoopInitComplete);
+	}
 }
 
 void FOculusEditorModule::ShutdownModule()
@@ -91,19 +93,20 @@ void FOculusEditorModule::ShutdownModule()
 		return;
 	}
 
-	if(!IsRunningCommandlet())
+	if (!IsRunningCommandlet())
 	{
 		FOculusToolStyle::Shutdown();
 		FOculusToolCommands::Unregister();
 		FGlobalTabmanager::Get()->UnregisterNomadTabSpawner(OculusPerfTabName);
 		FGlobalTabmanager::Get()->UnregisterNomadTabSpawner(OculusPlatToolTabName);
+		FOculusBuildAnalytics::Shutdown();
 	}
 
 	FOculusAssetDirectory::ReleaseAll();
 	if (UObjectInitialized())
 	{
-		UnregisterSettings();  
-	}		
+		UnregisterSettings();
+	}
 }
 
 TSharedRef<SDockTab> FOculusEditorModule::OnSpawnPluginTab(const FSpawnTabArgs& SpawnTabArgs)
@@ -178,6 +181,11 @@ void FOculusEditorModule::AddToolbarExtension(FToolBarBuilder& Builder)
 	Builder.AddToolBarButton(FOculusToolCommands::Get().OpenPluginWindow);
 }
 
+void FOculusEditorModule::OnEngineLoopInitComplete()
+{
+	BuildAnalytics = FOculusBuildAnalytics::GetInstance();
+}
+
 TSharedRef<IDetailCustomization> FOculusHMDSettingsDetailsCustomization::MakeInstance()
 {
 	return MakeShareable(new FOculusHMDSettingsDetailsCustomization);
@@ -195,6 +203,20 @@ FReply FOculusHMDSettingsDetailsCustomization::PluginClickPlatFn(bool text)
 	return FReply::Handled();
 }
 
+void FOculusHMDSettingsDetailsCustomization::OnEnableBuildTelemetry(ECheckBoxState NewState)
+{
+	FOculusBuildAnalytics::GetInstance()->OnTelemetryToggled(NewState == ECheckBoxState::Checked);
+	GConfig->SetBool(TEXT("/Script/OculusEditor.OculusEditorSettings"), TEXT("bEnableOculusBuildTelemetry"), NewState == ECheckBoxState::Checked, GEditorIni);
+	GConfig->Flush(0);
+}
+
+ECheckBoxState FOculusHMDSettingsDetailsCustomization::IsBuildTelemtryEnabled() const
+{
+	bool v;
+	GConfig->GetBool(TEXT("/Script/OculusEditor.OculusEditorSettings"), TEXT("bEnableOculusBuildTelemetry"), v, GEditorIni);
+	return v ? ECheckBoxState::Checked : ECheckBoxState::Unchecked;
+}
+
 void FOculusHMDSettingsDetailsCustomization::CustomizeDetails(IDetailLayoutBuilder& DetailLayout)
 {
 	// Labeled "General Oculus" instead of "General" to enable searchability. The button "Launch Oculus Utilities Window" doesn't show up if you search for "Oculus"
@@ -204,27 +226,50 @@ void FOculusHMDSettingsDetailsCustomization::CustomizeDetails(IDetailLayoutBuild
 		[
 			SNew(SVerticalBox)
 			+ SVerticalBox::Slot().AutoHeight().Padding(2)
+		[
+			SNew(SHorizontalBox)
+			+ SHorizontalBox::Slot().AutoWidth()
+		[
+			SNew(SButton)
+			.Text(LOCTEXT("LaunchTool", "Launch Oculus Performance Window"))
+		.OnClicked(this, &FOculusHMDSettingsDetailsCustomization::PluginClickPerfFn, true)
+		]
+	+ SHorizontalBox::Slot().FillWidth(8)
+		]
+	+ SVerticalBox::Slot().AutoHeight().Padding(2)
+		[
+			SNew(SHorizontalBox)
+			+ SHorizontalBox::Slot().AutoWidth()
+		[
+			SNew(SButton)
+			.Text(LOCTEXT("LaunchPlatTool", "Launch Oculus Platform Window"))
+		.OnClicked(this, &FOculusHMDSettingsDetailsCustomization::PluginClickPlatFn, true)
+		]
+	+ SHorizontalBox::Slot().FillWidth(8)
+		]
+	+ SVerticalBox::Slot().AutoHeight().Padding(5)
+		[
+			SNew(SHorizontalBox)
+			+ SHorizontalBox::Slot().AutoWidth()
+		[
+			SNew(SBox)
 			[
-				SNew(SHorizontalBox)
-				+ SHorizontalBox::Slot().AutoWidth()
-				[
-					SNew(SButton)
-					.Text(LOCTEXT("LaunchTool", "Launch Oculus Performance Window"))
-					.OnClicked(this, &FOculusHMDSettingsDetailsCustomization::PluginClickPerfFn, true)
-				]
-				+ SHorizontalBox::Slot().FillWidth(8)
-			]
-			+ SVerticalBox::Slot().AutoHeight().Padding(2)
-			[
-				SNew(SHorizontalBox)
-				+ SHorizontalBox::Slot().AutoWidth()
-				[
-					SNew(SButton)
-					.Text(LOCTEXT("LaunchPlatTool", "Launch Oculus Platform Window"))
-					.OnClicked(this, &FOculusHMDSettingsDetailsCustomization::PluginClickPlatFn, true)
-				]
-				+ SHorizontalBox::Slot().FillWidth(8)
+				SNew(STextBlock)
+				.Text(LOCTEXT("EnableBuildTelemetry", "Enable Oculus Build Telemetry"))
 			]
+		]
+	+ SHorizontalBox::Slot().AutoWidth()
+		[
+			SNew(SBox).WidthOverride(10.f)
+		]
+	+ SHorizontalBox::Slot().AutoWidth()
+		[
+			SNew(SCheckBox)
+			.OnCheckStateChanged(this, &FOculusHMDSettingsDetailsCustomization::OnEnableBuildTelemetry)
+		.IsChecked(this, &FOculusHMDSettingsDetailsCustomization::IsBuildTelemtryEnabled)
+		]
+	+ SHorizontalBox::Slot().FillWidth(8)
+		]
 		];
 }
 
diff --git a/Engine/Plugins/Runtime/Oculus/OculusVR/Source/OculusEditor/Private/OculusEditorModule.h b/Engine/Plugins/Runtime/Oculus/OculusVR/Source/OculusEditor/Private/OculusEditorModule.h
--- a/Engine/Plugins/Runtime/Oculus/OculusVR/Source/OculusEditor/Private/OculusEditorModule.h
+++ b/Engine/Plugins/Runtime/Oculus/OculusVR/Source/OculusEditor/Private/OculusEditorModule.h
@@ -4,6 +4,7 @@
 
 #include "CoreMinimal.h"
 #include "IOculusEditorModule.h"
+#include "OculusBuildAnalytics.h"
 #include "Modules/ModuleInterface.h"
 #include "IDetailCustomization.h"
 #include "Input/Reply.h"
@@ -26,10 +27,12 @@ public:
 
 	void RegisterSettings();
 	void UnregisterSettings();
-	
+
 	void PluginButtonClicked();
 	FReply PluginClickFn(bool text);
 
+	void OnEngineLoopInitComplete();
+
 public:
 	static const FName OculusPerfTabName;
 	static const FName OculusPlatToolTabName;
@@ -45,6 +48,7 @@ private:
 private:
 	TSharedPtr<class FUICommandList> PluginCommands;
 	bool bModuleValid;
+	FOculusBuildAnalytics* BuildAnalytics;
 };
 
 class IDetailLayoutBuilder;
@@ -61,4 +65,7 @@ public:
 
 	FReply PluginClickPerfFn(bool text);
 	FReply PluginClickPlatFn(bool text);
+
+	void OnEnableBuildTelemetry(ECheckBoxState NewState);
+	ECheckBoxState IsBuildTelemtryEnabled() const;
 };
diff --git a/Engine/Source/Developer/LauncherServices/Private/Launcher/Launcher.cpp b/Engine/Source/Developer/LauncherServices/Private/Launcher/Launcher.cpp
--- a/Engine/Source/Developer/LauncherServices/Private/Launcher/Launcher.cpp
+++ b/Engine/Source/Developer/LauncherServices/Private/Launcher/Launcher.cpp
@@ -24,8 +24,10 @@ ILauncherWorkerPtr FLauncher::Launch(const TSharedRef<ITargetDeviceProxyManager>
 
 		if ((LauncherWorker != nullptr) && (FRunnableThread::Create(LauncherWorker, *WorkerName) != nullptr))
 		{
-			return MakeShareable(LauncherWorker);
-		}			
+			ILauncherWorkerPtr Worker = MakeShareable(LauncherWorker);
+			FLauncherWorkerStartedDelegate.Broadcast(Worker, Profile);
+			return Worker;
+		}
 	}
 
 	return nullptr;
diff --git a/Engine/Source/Developer/LauncherServices/Private/LauncherServicesModule.cpp b/Engine/Source/Developer/LauncherServices/Private/LauncherServicesModule.cpp
--- a/Engine/Source/Developer/LauncherServices/Private/LauncherServicesModule.cpp
+++ b/Engine/Source/Developer/LauncherServices/Private/LauncherServicesModule.cpp
@@ -32,7 +32,9 @@ public:
 
 	virtual ILauncherRef CreateLauncher() override
 	{
-		return MakeShareable(new FLauncher());
+		ILauncherRef Launcher = MakeShareable(new FLauncher());
+		OnCreateLauncherDelegate.Broadcast(Launcher);
+		return Launcher;
 	}
 
 	virtual ILauncherProfileRef CreateProfile(const FString& ProfileName) override
diff --git a/Engine/Source/Developer/LauncherServices/Public/ILauncher.h b/Engine/Source/Developer/LauncherServices/Public/ILauncher.h
--- a/Engine/Source/Developer/LauncherServices/Public/ILauncher.h
+++ b/Engine/Source/Developer/LauncherServices/Public/ILauncher.h
@@ -15,6 +15,8 @@ typedef TSharedPtr<class ILauncher> ILauncherPtr;
 /** Type definition for shared references to instances of ILauncher. */
 typedef TSharedRef<class ILauncher> ILauncherRef;
 
+DECLARE_MULTICAST_DELEGATE_TwoParams(FOnLauncherWorkerStarted, ILauncherWorkerPtr, ILauncherProfileRef);
+
 
 /**
  * Interface for game launchers.
@@ -36,4 +38,6 @@ public:
 
 	/** Virtual destructor. */
 	virtual ~ILauncher( ) { }
+
+	FOnLauncherWorkerStarted FLauncherWorkerStartedDelegate;
 };
diff --git a/Engine/Source/Developer/LauncherServices/Public/ILauncherServicesModule.h b/Engine/Source/Developer/LauncherServices/Public/ILauncherServicesModule.h
--- a/Engine/Source/Developer/LauncherServices/Public/ILauncherServicesModule.h
+++ b/Engine/Source/Developer/LauncherServices/Public/ILauncherServicesModule.h
@@ -14,6 +14,8 @@
  */
 DECLARE_MULTICAST_DELEGATE_OneParam(FOnLauncherProfileManagerInitialized, ILauncherProfileManager&);
 
+DECLARE_MULTICAST_DELEGATE_OneParam(FOnCreateLauncher, ILauncherRef);
+
 
 /**
  * Interface for launcher tools modules.
@@ -80,4 +82,6 @@ public:
 	 * Delegate that is invoked when a profile manager is initialized.
 	 */
 	LAUNCHERSERVICES_API static FOnLauncherProfileManagerInitialized ProfileManagerInitializedDelegate;
+
+	FOnCreateLauncher OnCreateLauncherDelegate;
 };
