Fix non-unity builds (and also fix a spelling error in a function name).

diff --git a/Engine/Plugins/Runtime/Oculus/OculusVR/Source/OculusEditor/Private/OculusEditorModule.cpp b/Engine/Plugins/Runtime/Oculus/OculusVR/Source/OculusEditor/Private/OculusEditorModule.cpp
--- a/Engine/Plugins/Runtime/Oculus/OculusVR/Source/OculusEditor/Private/OculusEditorModule.cpp
+++ b/Engine/Plugins/Runtime/Oculus/OculusVR/Source/OculusEditor/Private/OculusEditorModule.cpp
@@ -210,7 +210,7 @@ void FOculusHMDSettingsDetailsCustomization::OnEnableBuildTelemetry(ECheckBoxSta
 	GConfig->Flush(0);
 }
 
-ECheckBoxState FOculusHMDSettingsDetailsCustomization::IsBuildTelemtryEnabled() const
+ECheckBoxState FOculusHMDSettingsDetailsCustomization::IsBuildTelemetryEnabled() const
 {
 	bool v;
 	GConfig->GetBool(TEXT("/Script/OculusEditor.OculusEditorSettings"), TEXT("bEnableOculusBuildTelemetry"), v, GEditorIni);
@@ -266,7 +266,7 @@ void FOculusHMDSettingsDetailsCustomization::CustomizeDetails(IDetailLayoutBuild
 		[
 			SNew(SCheckBox)
 			.OnCheckStateChanged(this, &FOculusHMDSettingsDetailsCustomization::OnEnableBuildTelemetry)
-		.IsChecked(this, &FOculusHMDSettingsDetailsCustomization::IsBuildTelemtryEnabled)
+		.IsChecked(this, &FOculusHMDSettingsDetailsCustomization::IsBuildTelemetryEnabled)
 		]
 	+ SHorizontalBox::Slot().FillWidth(8)
 		]
diff --git a/Engine/Plugins/Runtime/Oculus/OculusVR/Source/OculusEditor/Private/OculusEditorModule.h b/Engine/Plugins/Runtime/Oculus/OculusVR/Source/OculusEditor/Private/OculusEditorModule.h
--- a/Engine/Plugins/Runtime/Oculus/OculusVR/Source/OculusEditor/Private/OculusEditorModule.h
+++ b/Engine/Plugins/Runtime/Oculus/OculusVR/Source/OculusEditor/Private/OculusEditorModule.h
@@ -14,6 +14,7 @@ class FMenuBuilder;
 
 #define OCULUS_EDITOR_MODULE_NAME "OculusEditor"
 
+enum class ECheckBoxState : uint8;
 
 class FOculusEditorModule : public IOculusEditorModule
 {
@@ -67,5 +68,5 @@ public:
 	FReply PluginClickPlatFn(bool text);
 
 	void OnEnableBuildTelemetry(ECheckBoxState NewState);
-	ECheckBoxState IsBuildTelemtryEnabled() const;
+	ECheckBoxState IsBuildTelemetryEnabled() const;
 };
