Fix static build analysis warnings (unchecked pointer) and related cleanup

diff --git a/Engine/Plugins/Runtime/Oculus/OculusVR/Source/OculusEditor/Private/OculusPlatformToolWidget.cpp b/Engine/Plugins/Runtime/Oculus/OculusVR/Source/OculusEditor/Private/OculusPlatformToolWidget.cpp
--- a/Engine/Plugins/Runtime/Oculus/OculusVR/Source/OculusEditor/Private/OculusPlatformToolWidget.cpp
+++ b/Engine/Plugins/Runtime/Oculus/OculusVR/Source/OculusEditor/Private/OculusPlatformToolWidget.cpp
@@ -1085,22 +1085,22 @@ FReply SOculusPlatformToolWidget::OnClearLaunchFilePath()
 
 FReply SOculusPlatformToolWidget::OnSelect2DLaunchPath()
 {
-	TSharedPtr<SWindow> parentWindow = FSlateApplication::Get().FindWidgetWindow(AsShared());
-	const void* parentWindowHandle = (parentWindow.IsValid() && parentWindow->GetNativeWindow().IsValid()) ? parentWindow->GetNativeWindow()->GetOSWindowHandle() : nullptr;
-
-	TArray<FString> path;
-	FString defaultPath = PlatformSettings->OculusRift2DLaunchPath.IsEmpty() ? FPaths::ProjectContentDir() : PlatformSettings->OculusRift2DLaunchPath;
-	if (FDesktopPlatformModule::Get()->OpenFileDialog(parentWindowHandle, "Choose 2D Launch File", defaultPath, defaultPath, "Executables (*.exe)|*.exe", EFileDialogFlags::None, path))
+	if (PlatformSettings != NULL)
 	{
-		if (PlatformSettings != NULL)
+		TSharedPtr<SWindow> parentWindow = FSlateApplication::Get().FindWidgetWindow(AsShared());
+		const void* parentWindowHandle = (parentWindow.IsValid() && parentWindow->GetNativeWindow().IsValid()) ? parentWindow->GetNativeWindow()->GetOSWindowHandle() : nullptr;
+
+		TArray<FString> path;
+		FString defaultPath = PlatformSettings->OculusRift2DLaunchPath.IsEmpty() ? FPaths::ProjectContentDir() : PlatformSettings->OculusRift2DLaunchPath;
+		if (FDesktopPlatformModule::Get()->OpenFileDialog(parentWindowHandle, "Choose 2D Launch File", defaultPath, defaultPath, "Executables (*.exe)|*.exe", EFileDialogFlags::None, path))
 		{
 			if (path.Num() > 0)
 			{
 				PlatformSettings->OculusRift2DLaunchPath = FPaths::ConvertRelativePathToFull(path[0]);
 			}
+			PlatformSettings->SaveConfig();
+			BuildRiftOptionalFields(OptionalSettings);
 		}
-		PlatformSettings->SaveConfig();
-		BuildRiftOptionalFields(OptionalSettings);
 	}
 	return FReply::Handled();
 }
@@ -1131,19 +1131,19 @@ FReply SOculusPlatformToolWidget::OnCancelUpload()
 
 FReply SOculusPlatformToolWidget::OnSelectLanguagePacksPath()
 {
-	TSharedPtr<SWindow> parentWindow = FSlateApplication::Get().FindWidgetWindow(AsShared());
-	const void* parentWindowHandle = (parentWindow.IsValid() && parentWindow->GetNativeWindow().IsValid()) ? parentWindow->GetNativeWindow()->GetOSWindowHandle() : nullptr;
-
-	FString path;
-	FString defaultPath = PlatformSettings->GetLanguagePacksPath().IsEmpty() ? FPaths::ProjectContentDir() : PlatformSettings->GetLanguagePacksPath();
-	if (FDesktopPlatformModule::Get()->OpenDirectoryDialog(parentWindowHandle, "Choose Language Packs Directory", defaultPath, path))
+	if (PlatformSettings != NULL)
 	{
-		if (PlatformSettings != NULL)
+		TSharedPtr<SWindow> parentWindow = FSlateApplication::Get().FindWidgetWindow(AsShared());
+		const void* parentWindowHandle = (parentWindow.IsValid() && parentWindow->GetNativeWindow().IsValid()) ? parentWindow->GetNativeWindow()->GetOSWindowHandle() : nullptr;
+
+		FString path;
+		FString defaultPath = PlatformSettings->GetLanguagePacksPath().IsEmpty() ? FPaths::ProjectContentDir() : PlatformSettings->GetLanguagePacksPath();
+		if (FDesktopPlatformModule::Get()->OpenDirectoryDialog(parentWindowHandle, "Choose Language Packs Directory", defaultPath, path))
 		{
 			PlatformSettings->SetLanguagePacksPath(path);
+			PlatformSettings->SaveConfig();
+			BuildExpansionFileBox(ExpansionFilesSettings);
 		}
-		PlatformSettings->SaveConfig();
-		BuildExpansionFileBox(ExpansionFilesSettings);
 	}
 	return FReply::Handled();
 }
@@ -1161,35 +1161,35 @@ FReply SOculusPlatformToolWidget::OnClearLanguagePacksPath()
 
 FReply SOculusPlatformToolWidget::OnSelectExpansionFilesPath()
 {
-	TSharedPtr<SWindow> parentWindow = FSlateApplication::Get().FindWidgetWindow(AsShared());
-	const void* parentWindowHandle = (parentWindow.IsValid() && parentWindow->GetNativeWindow().IsValid()) ? parentWindow->GetNativeWindow()->GetOSWindowHandle() : nullptr;
-
-	FString path;
-	FString defaultPath = PlatformSettings->GetExpansionFilesPath().IsEmpty() ? FPaths::ProjectContentDir() : PlatformSettings->GetExpansionFilesPath();
-	if (FDesktopPlatformModule::Get()->OpenDirectoryDialog(parentWindowHandle, "Choose Expansion Files Directory", defaultPath, path))
+	if (PlatformSettings != NULL)
 	{
-		if (!path.Equals(PlatformSettings->GetExpansionFilesPath()))
+		TSharedPtr<SWindow> parentWindow = FSlateApplication::Get().FindWidgetWindow(AsShared());
+		const void* parentWindowHandle = (parentWindow.IsValid() && parentWindow->GetNativeWindow().IsValid()) ? parentWindow->GetNativeWindow()->GetOSWindowHandle() : nullptr;
+
+		FString path;
+		FString defaultPath = PlatformSettings->GetExpansionFilesPath().IsEmpty() ? FPaths::ProjectContentDir() : PlatformSettings->GetExpansionFilesPath();
+		if (FDesktopPlatformModule::Get()->OpenDirectoryDialog(parentWindowHandle, "Choose Expansion Files Directory", defaultPath, path))
 		{
-			if (!path.IsEmpty() && FPaths::DirectoryExists(path))
+			if (!path.Equals(PlatformSettings->GetExpansionFilesPath()))
 			{
-				TArray<FString> Files;
-				//FFileManagerGeneric::Get().FindFilesRecursive(Files, *path, TEXT("*.*"), true, false, false);
-				IFileManager::Get().FindFiles(Files, *path);
-
-				TArray<FAssetConfig>* AssetConfigs = PlatformSettings->GetAssetConfigs();
-				for (int i = 0; i < Files.Num(); i++)
+				if (!path.IsEmpty() && FPaths::DirectoryExists(path))
 				{
-					FAssetConfig AssetConfig;
-					AssetConfig.Name = Files[i];
-					AssetConfigs->Push(AssetConfig);
-				}
+					TArray<FString> Files;
+					//FFileManagerGeneric::Get().FindFilesRecursive(Files, *path, TEXT("*.*"), true, false, false);
+					IFileManager::Get().FindFiles(Files, *path);
+
+					TArray<FAssetConfig>* AssetConfigs = PlatformSettings->GetAssetConfigs();
+					for (int i = 0; i < Files.Num(); i++)
+					{
+						FAssetConfig AssetConfig;
+						AssetConfig.Name = Files[i];
+						AssetConfigs->Push(AssetConfig);
+					}
 
-				if (PlatformSettings != NULL)
-				{
 					PlatformSettings->SetExpansionFilesPath(path);
+					PlatformSettings->SaveConfig();
+					BuildExpansionFileBox(ExpansionFilesSettings);
 				}
-				PlatformSettings->SaveConfig();
-				BuildExpansionFileBox(ExpansionFilesSettings);
 			}
 		}
 	}
