Fix for VR mode not working on Oculus

diff --git a/Engine/Plugins/Runtime/Oculus/OculusVR/Source/OculusInput/Private/OculusInput.cpp b/Engine/Plugins/Runtime/Oculus/OculusVR/Source/OculusInput/Private/OculusInput.cpp
--- a/Engine/Plugins/Runtime/Oculus/OculusVR/Source/OculusInput/Private/OculusInput.cpp
+++ b/Engine/Plugins/Runtime/Oculus/OculusVR/Source/OculusInput/Private/OculusInput.cpp
@@ -758,7 +758,7 @@ bool FOculusInput::GetControllerOrientationAndPosition( const int32 ControllerIn
 						OculusHMD::FOculusHMD* OculusHMD = static_cast<OculusHMD::FOculusHMD*>(GEngine->XRSystem->GetHMDDevice());
 						ovrpNode Node = DeviceHand == EControllerHand::Left ? ovrpNode_HandLeft : ovrpNode_HandRight;
 
-						ovrpBool bResult;
+						ovrpBool bResult = true;
 						bool bIsPositionValid = OVRP_SUCCESS(ovrp_GetNodePositionValid(Node, &bResult)) && bResult;
 						bool bIsOrientationValid = OVRP_SUCCESS(ovrp_GetNodeOrientationValid(Node, &bResult)) && bResult;
 
@@ -770,7 +770,7 @@ bool FOculusInput::GetControllerOrientationAndPosition( const int32 ControllerIn
 							if (IsInGameThread())
 							{
 								Settings = OculusHMD->GetSettings();
-								CurrentFrame = OculusHMD->GetFrame();
+								CurrentFrame = OculusHMD->GetNextFrameToRender();
 							}
 							else
 							{
