Fix for a crash during shutdown caused by loading screen refactor

diff --git a/Engine/Plugins/Runtime/Oculus/OculusVR/Source/OculusHMD/Private/OculusHMD.cpp b/Engine/Plugins/Runtime/Oculus/OculusVR/Source/OculusHMD/Private/OculusHMD.cpp
--- a/Engine/Plugins/Runtime/Oculus/OculusVR/Source/OculusHMD/Private/OculusHMD.cpp
+++ b/Engine/Plugins/Runtime/Oculus/OculusVR/Source/OculusHMD/Private/OculusHMD.cpp
@@ -2074,6 +2074,8 @@ namespace OculusHMD
 		{
 			Splash->Shutdown();
 			Splash = nullptr;
+			// The base implementation stores a raw pointer to the Splash object and tries to deallocate it in its destructor
+			LoadingScreen = nullptr;
 		}
 
 		if (CustomPresent.IsValid())
