Fix localization warning.

#rnx

diff --git a/Engine/Plugins/Runtime/Oculus/OculusVR/Source/OculusEditor/Private/OculusPlatformToolWidget.cpp b/Engine/Plugins/Runtime/Oculus/OculusVR/Source/OculusEditor/Private/OculusPlatformToolWidget.cpp
--- a/Engine/Plugins/Runtime/Oculus/OculusVR/Source/OculusEditor/Private/OculusPlatformToolWidget.cpp
+++ b/Engine/Plugins/Runtime/Oculus/OculusVR/Source/OculusEditor/Private/OculusPlatformToolWidget.cpp
@@ -499,7 +499,7 @@ void SOculusPlatformToolWidget::BuildExpansionFileBox(TSharedPtr<SVerticalBox> b
 			LOCTEXT("LanguagePacksTT", ""), &SOculusPlatformToolWidget::OnSelectLanguagePacksPath, &SOculusPlatformToolWidget::OnClearLanguagePacksPath);
 	}
 
-	BuildFileDirectoryField(box, LOCTEXT("ExpansionFiles", "Expansion Files Directory"), FText::FromString(PlatformSettings->GetExpansionFilesPath()),
+	BuildFileDirectoryField(box, LOCTEXT("ExpansionFilesDirectory", "Expansion Files Directory"), FText::FromString(PlatformSettings->GetExpansionFilesPath()),
 		LOCTEXT("ExpansionFilesTT", ""), &SOculusPlatformToolWidget::OnSelectExpansionFilesPath, &SOculusPlatformToolWidget::OnClearExpansionFilesPath);
 
 	TArray<FAssetConfig>* AssetConfigs = PlatformSettings->GetAssetConfigs();
@@ -738,7 +738,7 @@ bool SOculusPlatformToolWidget::ConstructArguments(FString& args)
 	if (!PlatformSettings->GetExpansionFilesPath().IsEmpty())
 	{
 		ValidateTextField(&SOculusPlatformToolWidget::DirectoryFieldValidator, PlatformSettings->GetExpansionFilesPath(),
-			LOCTEXT("ExpansionFiles", "Expansion Files Path").ToString(), success);
+			LOCTEXT("ExpansionFilesPath", "Expansion Files Path").ToString(), success);
 		args += " --assets-dir \"" + PlatformSettings->GetExpansionFilesPath() + "\"";
 
 		TArray<FAssetConfig>* AssetConfigs = PlatformSettings->GetAssetConfigs();
