Manually copying changes in CL#9747937 from Dev-VR:
Fix hitching that occurs on Oculus when a level contains unitialized or hidden stereo layers when entering VR preview.

diff --git a/Engine/Plugins/Runtime/Oculus/OculusVR/Source/OculusHMD/Private/OculusHMD.cpp b/Engine/Plugins/Runtime/Oculus/OculusVR/Source/OculusHMD/Private/OculusHMD.cpp
--- a/Engine/Plugins/Runtime/Oculus/OculusVR/Source/OculusHMD/Private/OculusHMD.cpp
+++ b/Engine/Plugins/Runtime/Oculus/OculusVR/Source/OculusHMD/Private/OculusHMD.cpp
@@ -3227,11 +3227,15 @@ namespace OculusHMD
 			FGameFramePtr XFrame = NextFrameToRender->Clone();
 			TArray<FLayerPtr> XLayers;
 
-			LayerMap.GenerateValueArray(XLayers);
+			XLayers.Empty(LayerMap.Num());
 
-			for (int32 XLayerIndex = 0; XLayerIndex < XLayers.Num(); XLayerIndex++)
+			for (auto Pair : LayerMap)
 			{
-				XLayers[XLayerIndex] = XLayers[XLayerIndex]->Clone();
+				// Skip hidden layers
+				if (!(Pair.Value->GetDesc().Flags & IStereoLayers::LAYER_FLAG_HIDDEN))
+				{
+					XLayers.Emplace(Pair.Value->Clone());
+				}
 			}
 
 			XLayers.Sort(FLayerPtr_CompareId());
diff --git a/Engine/Plugins/Runtime/Oculus/OculusVR/Source/OculusHMD/Private/OculusHMD_Layer.cpp b/Engine/Plugins/Runtime/Oculus/OculusVR/Source/OculusHMD/Private/OculusHMD_Layer.cpp
--- a/Engine/Plugins/Runtime/Oculus/OculusVR/Source/OculusHMD/Private/OculusHMD_Layer.cpp
+++ b/Engine/Plugins/Runtime/Oculus/OculusVR/Source/OculusHMD/Private/OculusHMD_Layer.cpp
@@ -343,10 +343,7 @@ void FLayer::Initialize_RenderThread(const FSettings* Settings, FCustomPresent*
 		bInvertY = (CustomPresent->GetLayerFlags() & ovrpLayerFlag_TextureOriginAtBottomLeft) != 0;
 
 		uint32 SizeX = 0, SizeY = 0;
-		if (Desc.Flags & IStereoLayers::LAYER_FLAG_HIDDEN)
-		{
-			return;
-		}
+		check((Desc.Flags & IStereoLayers::LAYER_FLAG_HIDDEN) == 0);
 
 		if (Desc.Texture.IsValid())
 		{
@@ -603,6 +600,7 @@ void FLayer::Initialize_RenderThread(const FSettings* Settings, FCustomPresent*
 void FLayer::UpdateTexture_RenderThread(FCustomPresent* CustomPresent, FRHICommandListImmediate& RHICmdList)
 {
 	CheckInRenderThread();
+	check((Desc.Flags & IStereoLayers::LAYER_FLAG_HIDDEN) == 0);
 
 	if (bUpdateTexture && SwapChain.IsValid())
 	{
@@ -652,6 +650,7 @@ void FLayer::UpdateTexture_RenderThread(FCustomPresent* CustomPresent, FRHIComma
 
 const ovrpLayerSubmit* FLayer::UpdateLayer_RHIThread(const FSettings* Settings, const FGameFrame* Frame, const int LayerIndex)
 {
+	check((Desc.Flags & IStereoLayers::LAYER_FLAG_HIDDEN) == 0);
 	OvrpLayerSubmit.LayerId = OvrpLayerId;
 	OvrpLayerSubmit.TextureStage = SwapChain.IsValid() ? SwapChain->GetSwapChainIndex_RHIThread() : 0;
 
@@ -785,6 +784,7 @@ const ovrpLayerSubmit* FLayer::UpdateLayer_RHIThread(const FSettings* Settings,
 
 void FLayer::IncrementSwapChainIndex_RHIThread(FCustomPresent* CustomPresent)
 {
+	check((Desc.Flags & IStereoLayers::LAYER_FLAG_HIDDEN) == 0);
 	CheckInRHIThread();
 
 	if (SwapChain.IsValid())
