[Partner-Oculus] Fix for Oculus controllers not appearing in editor VRMode.

[CL 8702557 by [REDACTED_PERSON] in Partner-Oculus-4.23-Private]
[CL 9375071 by [REDACTED_PERSON] in Partner-Oculus-Private]

diff --git a/Engine/Plugins/Runtime/Oculus/OculusVR/Source/OculusInput/Private/OculusInput.cpp b/Engine/Plugins/Runtime/Oculus/OculusVR/Source/OculusInput/Private/OculusInput.cpp
--- a/Engine/Plugins/Runtime/Oculus/OculusVR/Source/OculusInput/Private/OculusInput.cpp
+++ b/Engine/Plugins/Runtime/Oculus/OculusVR/Source/OculusInput/Private/OculusInput.cpp
@@ -778,12 +778,12 @@ bool FOculusInput::GetControllerOrientationAndPosition( const int32 ControllerIn
 								CurrentFrame = OculusHMD->GetFrame_RenderThread();
 							}
 
-							if (Settings && CurrentFrame)
+							if (Settings)
 							{
 								ovrpPoseStatef InPoseState;
 								OculusHMD::FPose OutPose;
 
-								if (OVRP_SUCCESS(ovrp_GetNodePoseState3(ovrpStep_Render, CurrentFrame->FrameNumber, Node, &InPoseState)) &&
+								if (OVRP_SUCCESS(ovrp_GetNodePoseState3(ovrpStep_Render, CurrentFrame ? CurrentFrame->FrameNumber : OVRP_CURRENT_FRAMEINDEX, Node, &InPoseState)) &&
 									OculusHMD->ConvertPose_Internal(InPoseState.Pose, OutPose, Settings, WorldToMetersScale))
 								{
 									if (bIsPositionValid)
