[Partner-Oculus] Add function to check MRC availability

[CL 8645863 by kevin.xiao in Partner-Oculus-4.23-Private]
[CL 9375069 by [REDACTED_PERSON] in Partner-Oculus-Private]

diff --git a/Engine/Plugins/Runtime/Oculus/OculusVR/Source/OculusMR/Private/OculusMRFunctionLibrary.cpp b/Engine/Plugins/Runtime/Oculus/OculusVR/Source/OculusMR/Private/OculusMRFunctionLibrary.cpp
--- a/Engine/Plugins/Runtime/Oculus/OculusVR/Source/OculusMR/Private/OculusMRFunctionLibrary.cpp
+++ b/Engine/Plugins/Runtime/Oculus/OculusVR/Source/OculusMR/Private/OculusMRFunctionLibrary.cpp
@@ -154,6 +154,11 @@ void UOculusMRFunctionLibrary::SetTrackingReferenceComponent(USceneComponent* Co
 	}
 }
 
+bool UOculusMRFunctionLibrary::IsMrcEnabled()
+{
+	return FOculusMRModule::IsAvailable() && FOculusMRModule::Get().IsInitialized();
+}
+
 bool UOculusMRFunctionLibrary::IsMrcActive()
 {
 	return FOculusMRModule::IsAvailable() && FOculusMRModule::Get().IsActive();
diff --git a/Engine/Plugins/Runtime/Oculus/OculusVR/Source/OculusMR/Public/OculusMRFunctionLibrary.h b/Engine/Plugins/Runtime/Oculus/OculusVR/Source/OculusMR/Public/OculusMRFunctionLibrary.h
--- a/Engine/Plugins/Runtime/Oculus/OculusVR/Source/OculusMR/Public/OculusMRFunctionLibrary.h
+++ b/Engine/Plugins/Runtime/Oculus/OculusVR/Source/OculusMR/Public/OculusMRFunctionLibrary.h
@@ -32,7 +32,11 @@ public:
 	UFUNCTION(BlueprintCallable, Category = "OculusLibrary|MR")
 	static void SetTrackingReferenceComponent(USceneComponent* Component);
 
-	// Check if MRC is active (enabled and capturing)
+	// Check if MRC is enabled
+	UFUNCTION(BlueprintCallable, Category = "OculusLibrary|MR")
+	static bool IsMrcEnabled();
+
+	// Check if MRC is enabled and actively capturing
 	UFUNCTION(BlueprintCallable, Category = "OculusLibrary|MR")
 	static bool IsMrcActive();
 
