[Partner-Oculus] Change OculusVR plugin version to 1.41.0

diff --git a/Engine/Plugins/Runtime/Oculus/OculusVR/OculusVR.uplugin b/Engine/Plugins/Runtime/Oculus/OculusVR/OculusVR.uplugin
--- a/Engine/Plugins/Runtime/Oculus/OculusVR/OculusVR.uplugin
+++ b/Engine/Plugins/Runtime/Oculus/OculusVR/OculusVR.uplugin
@@ -1,7 +1,7 @@
 {
 	"FileVersion": 3,
 	"Version": 1,
-	"VersionName": "1.40.0",
+	"VersionName": "1.41.0",
 	"FriendlyName": "Oculus VR",
 	"Description": "Support for Oculus headsets and controllers",
 	"Category": "Virtual Reality",
