Merging using PartnerOculus_to_Release_4.24
(Fix Vulkan MRC crashes)
[CL 8639854 by kevin.xiao in Partner-Oculus-4.23-Private]
[CL 9375064 by [REDACTED_PERSON] in Partner-Oculus-Private]

diff --git a/Engine/Plugins/Runtime/Oculus/OculusVR/Source/OculusMR/OculusMR.Build.cs b/Engine/Plugins/Runtime/Oculus/OculusVR/Source/OculusMR/OculusMR.Build.cs
--- a/Engine/Plugins/Runtime/Oculus/OculusVR/Source/OculusMR/OculusMR.Build.cs
+++ b/Engine/Plugins/Runtime/Oculus/OculusVR/Source/OculusMR/OculusMR.Build.cs
@@ -25,6 +25,7 @@ namespace UnrealBuildTool.Rules
                     "Slate",
                     "SlateCore",
                     "RHI",
+                    "VulkanRHI",
                     "RenderCore",
                     "MediaAssets",
                     "HeadMountedDisplay",
@@ -49,6 +50,13 @@ namespace UnrealBuildTool.Rules
 				RuntimeDependencies.Add("$(EngineDir)/Binaries/ThirdParty/Oculus/OVRPlugin/OVRPlugin/" + Target.Platform.ToString() + "/OVRPlugin.dll");
 			}
 
+			if (Target.Platform == UnrealTargetPlatform.Android)
+            {
+                PrivateIncludePaths.Add("../../../../../Source/Runtime/VulkanRHI/Private");
+                PrivateIncludePaths.Add("../../../../../Source/Runtime/VulkanRHI/Private/" + Target.Platform);
+                AddEngineThirdPartyPrivateStaticDependencies(Target, "Vulkan");
+            }
+
             if (Target.bBuildEditor == true)
             {
                 PrivateDependencyModuleNames.Add("UnrealEd");
diff --git a/Engine/Source/Runtime/Renderer/Private/MobileShadingRenderer.cpp b/Engine/Source/Runtime/Renderer/Private/MobileShadingRenderer.cpp
--- a/Engine/Source/Runtime/Renderer/Private/MobileShadingRenderer.cpp
+++ b/Engine/Source/Runtime/Renderer/Private/MobileShadingRenderer.cpp
@@ -531,7 +531,7 @@ void FMobileSceneRenderer::Render(FRHICommandListImmediate& RHICmdList)
 
 	FRHITexture* FoveationTexture = nullptr;
 	
-	if (SceneContext.IsFoveationTextureAllocated())
+	if (SceneContext.IsFoveationTextureAllocated()	&& !View.bIsSceneCapture && !View.bIsReflectionCapture)
 	{
 		FoveationTexture = SceneContext.GetFoveationTexture();
 	}
diff --git a/Engine/Source/Runtime/VulkanRHI/Private/VulkanQueue.cpp b/Engine/Source/Runtime/VulkanRHI/Private/VulkanQueue.cpp
--- a/Engine/Source/Runtime/VulkanRHI/Private/VulkanQueue.cpp
+++ b/Engine/Source/Runtime/VulkanRHI/Private/VulkanQueue.cpp
@@ -21,12 +21,13 @@ FAutoConsoleVariableRef CVarVulkanWaitForIdleOnSubmit(
 FVulkanQueue::FVulkanQueue(FVulkanDevice* InDevice, uint32 InFamilyIndex)
 	: Queue(VK_NULL_HANDLE)
 	, FamilyIndex(InFamilyIndex)
+	, QueueIndex(0)
 	, Device(InDevice)
 	, LastSubmittedCmdBuffer(nullptr)
 	, LastSubmittedCmdBufferFenceCounter(0)
 	, SubmitCounter(0)
 {
-	VulkanRHI::vkGetDeviceQueue(Device->GetInstanceHandle(), FamilyIndex, 0, &Queue);
+	VulkanRHI::vkGetDeviceQueue(Device->GetInstanceHandle(), FamilyIndex, QueueIndex, &Queue);
 }
 
 FVulkanQueue::~FVulkanQueue()
diff --git a/Engine/Source/Runtime/VulkanRHI/Private/VulkanQueue.h b/Engine/Source/Runtime/VulkanRHI/Private/VulkanQueue.h
--- a/Engine/Source/Runtime/VulkanRHI/Private/VulkanQueue.h
+++ b/Engine/Source/Runtime/VulkanRHI/Private/VulkanQueue.h
@@ -1,4 +1,4 @@
-﻿// Copyright 1998-2019 Epic Games, Inc. All Rights Reserved..
+// Copyright 1998-2019 Epic Games, Inc. All Rights Reserved..
 
 /*=============================================================================
 	VulkanQueue.h: Private Vulkan RHI definitions.
@@ -29,6 +29,11 @@ public:
 		return FamilyIndex;
 	}
 
+	inline uint32 GetQueueIndex() const
+	{
+		return QueueIndex;
+	}
+
 	void Submit(FVulkanCmdBuffer* CmdBuffer, uint32 NumSignalSemaphores = 0, VkSemaphore* SignalSemaphores = nullptr);
 
 	inline void Submit(FVulkanCmdBuffer* CmdBuffer, VkSemaphore SignalSemaphore)
