Fix compile error
#rnx

diff --git a/Engine/Plugins/Runtime/Oculus/OculusVR/Source/OculusMR/Private/OculusMR_CastingCameraActor.cpp b/Engine/Plugins/Runtime/Oculus/OculusVR/Source/OculusMR/Private/OculusMR_CastingCameraActor.cpp
--- a/Engine/Plugins/Runtime/Oculus/OculusVR/Source/OculusMR/Private/OculusMR_CastingCameraActor.cpp
+++ b/Engine/Plugins/Runtime/Oculus/OculusVR/Source/OculusMR/Private/OculusMR_CastingCameraActor.cpp
@@ -281,8 +281,9 @@ void AOculusMR_CastingCameraActor::BeginPlay()
 void AOculusMR_CastingCameraActor::EndPlay(EEndPlayReason::Type Reason)
 {
 #if PLATFORM_ANDROID
+	FAudioDevice* AudioDevice = FAudioDeviceManager::GetMainDevice();
 	float NumChannels = 2;
-	float SampleRate = FAudioDeviceManager::GetMainDevice()->GetSampleRate();
+	float SampleRate = AudioDevice->GetSampleRate();
 	AudioDevice->StopRecording(nullptr, NumChannels, SampleRate);
 #endif
 
