Fix additional compile error
#rnx

diff --git a/Engine/Plugins/Runtime/Oculus/OculusVR/Source/OculusMR/Private/OculusMRModule.cpp b/Engine/Plugins/Runtime/Oculus/OculusVR/Source/OculusMR/Private/OculusMRModule.cpp
--- a/Engine/Plugins/Runtime/Oculus/OculusVR/Source/OculusMR/Private/OculusMRModule.cpp
+++ b/Engine/Plugins/Runtime/Oculus/OculusVR/Source/OculusMR/Private/OculusMRModule.cpp
@@ -447,8 +447,7 @@ void FOculusMRModule::OnInitialWorldCreated(UWorld* NewWorld)
 				ovrp_Media_SetMrcInputVideoBufferType(ovrpMediaInputVideoBufferType_TextureHandle);
 				ovrp_Media_SetMrcFrameImageFlipped(IsAndroidOpenGLESPlatform(GMaxRHIShaderPlatform));
 				ovrp_Media_SetMrcFrameInverseAlpha(ovrpBool_True);
-				FAudioDevice* AudioDevice = FAudioDevice::GetMainAudioDevice();
-				float SampleRate = AudioDevice->GetSampleRate();
+				float SampleRate = FAudioDeviceManager::GetMainDevice()->GetSampleRate();
 				ovrp_Media_SetMrcAudioSampleRate((int)SampleRate);
 				InitMixedRealityCapture();
 				OnWorldCreated(NewWorld);
