Refactor stereo layer shapes.
Part1, the C++ side.

UStereoLayerComponent is still monolithic, but has been updated to use the new C++ API in the implementation. -- Updated component with custom Serialize method, and all that will be in part 2 ;)

diff --git a/Engine/Plugins/Runtime/Oculus/OculusVR/Source/OculusHMD/Private/OculusHMD.cpp b/Engine/Plugins/Runtime/Oculus/OculusVR/Source/OculusHMD/Private/OculusHMD.cpp
--- a/Engine/Plugins/Runtime/Oculus/OculusVR/Source/OculusHMD/Private/OculusHMD.cpp
+++ b/Engine/Plugins/Runtime/Oculus/OculusVR/Source/OculusHMD/Private/OculusHMD.cpp
@@ -1758,9 +1758,10 @@ namespace OculusHMD
 		if (LayerFound && (*LayerFound)->GetSwapChain().IsValid())
 		{
 			bool bRightTexture = (*LayerFound)->GetRightSwapChain().IsValid();
-			switch ((*LayerFound)->GetDesc().ShapeType)
+			const IStereoLayers::FLayerDesc& Desc = (*LayerFound)->GetDesc();
+			
+			if (Desc.HasShape<FCubemapLayer>())
 			{
-			case IStereoLayers::CubemapLayer:
 				if (bRightTexture)
 				{
 					Texture = (*LayerFound)->GetRightSwapChain()->GetTextureCube();
@@ -1769,10 +1770,10 @@ namespace OculusHMD
 				else
 				{
 					Texture = LeftTexture = (*LayerFound)->GetSwapChain()->GetTextureCube();
-				}				break;
-
-			case IStereoLayers::CylinderLayer:
-			case IStereoLayers::QuadLayer:
+				}
+			}
+			else if (Desc.HasShape<FCylinderLayer>() || Desc.HasShape<FQuadLayer>())
+			{
 				if (bRightTexture)
 				{
 					Texture = (*LayerFound)->GetRightSwapChain()->GetTexture2D();
@@ -1782,24 +1783,16 @@ namespace OculusHMD
 				{
 					Texture = LeftTexture = (*LayerFound)->GetSwapChain()->GetTexture2D();
 				}
-				break;
-
-			default:
-				break;
 			}
 		}
 	}
 
 	IStereoLayers::FLayerDesc FOculusHMD::GetDebugCanvasLayerDesc(FTextureRHIRef Texture)
 	{
-		IStereoLayers::FLayerDesc StereoLayerDesc;
+		IStereoLayers::FLayerDesc StereoLayerDesc(FCylinderLayer(180.f, 488.f / 4, 100.f));
 		StereoLayerDesc.Transform = FTransform(FVector(0.f, 0, 0)); //100/0/0 for quads
-		StereoLayerDesc.CylinderHeight = 180.f;
-		StereoLayerDesc.CylinderOverlayArc = 488.f/4;
-		StereoLayerDesc.CylinderRadius = 100.f;
 		StereoLayerDesc.QuadSize = FVector2D(180.f, 180.f);
 		StereoLayerDesc.PositionType = IStereoLayers::ELayerType::FaceLocked;
-		StereoLayerDesc.ShapeType = IStereoLayers::ELayerShape::CylinderLayer;
 		StereoLayerDesc.LayerSize = Texture->GetTexture2D()->GetSizeXY();
 		StereoLayerDesc.Flags = IStereoLayers::ELayerFlags::LAYER_FLAG_TEX_CONTINUOUS_UPDATE;
 		StereoLayerDesc.Flags |= IStereoLayers::ELayerFlags::LAYER_FLAG_QUAD_PRESERVE_TEX_RATIO;
diff --git a/Engine/Plugins/Runtime/Oculus/OculusVR/Source/OculusHMD/Private/OculusHMD_Layer.cpp b/Engine/Plugins/Runtime/Oculus/OculusVR/Source/OculusHMD/Private/OculusHMD_Layer.cpp
--- a/Engine/Plugins/Runtime/Oculus/OculusVR/Source/OculusHMD/Private/OculusHMD_Layer.cpp
+++ b/Engine/Plugins/Runtime/Oculus/OculusVR/Source/OculusHMD/Private/OculusHMD_Layer.cpp
@@ -184,7 +184,7 @@ static void AppendFaceIndices(const int v0, const int v1, const int v2, const in
 
 void FLayer::BuildPokeAHoleMesh(TArray<FVector>& Vertices, TArray<int32>& Triangles, TArray<FVector2D>& UV0)
 {
-	if (Desc.ShapeType == IStereoLayers::QuadLayer)
+	if (Desc.HasShape<FQuadLayer>())
 	{
 		const float QuadScale = 0.99;
 
@@ -209,20 +209,21 @@ void FLayer::BuildPokeAHoleMesh(TArray<FVector>& Vertices, TArray<int32>& Triang
 		Triangles.Reserve(6);
 		AppendFaceIndices(0, 1, 2, 3, Triangles, false);
 	}
-	else if (Desc.ShapeType == IStereoLayers::CylinderLayer)
+	else if (Desc.HasShape<FCylinderLayer>())
 	{
+		const FCylinderLayer& CylinderProps = Desc.GetShape<FCylinderLayer>();
 		const float CylinderScale = 0.99;
 
 		FIntPoint TexSize = Desc.Texture.IsValid() ? Desc.Texture->GetTexture2D()->GetSizeXY() : Desc.LayerSize;
 		float AspectRatio = TexSize.X ? (float)TexSize.Y / (float)TexSize.X : 3.0f / 4.0f;
 
-		float CylinderHeight = (Desc.Flags & IStereoLayers::LAYER_FLAG_QUAD_PRESERVE_TEX_RATIO) ? Desc.CylinderOverlayArc * AspectRatio : Desc.CylinderHeight;
+		float CylinderHeight = (Desc.Flags & IStereoLayers::LAYER_FLAG_QUAD_PRESERVE_TEX_RATIO) ? CylinderProps->OverlayArc * AspectRatio : CylinderProps->Height;
 
 		const FVector XAxis = FVector(1, 0, 0);
 		const FVector YAxis = FVector(0, 1, 0);
 		const FVector HalfHeight = FVector(0, 0, CylinderHeight / 2);
 
-		const float ArcAngle = Desc.CylinderOverlayArc / Desc.CylinderRadius;
+		const float ArcAngle = CylinderProps->OverlayArc / CylinderProps->Radius;
 		const int Sides = (int)( (ArcAngle * 180) / (PI * 5) ); // one triangle every 10 degrees of cylinder for a good-cheap approximation
 		Vertices.Init(FVector::ZeroVector, 2 * (Sides + 1));
 		UV0.Init(FVector2D::ZeroVector, 2 * (Sides + 1));
@@ -234,7 +235,7 @@ void FLayer::BuildPokeAHoleMesh(TArray<FVector>& Vertices, TArray<int32>& Triang
 
 		for (int Side = 0; Side < Sides + 1; Side++)
 		{
-			FVector MidVertex = Desc.CylinderRadius * (FMath::Cos(CurrentAngle) * XAxis + FMath::Sin(CurrentAngle) * YAxis);
+			FVector MidVertex = CylinderProps->Radius * (FMath::Cos(CurrentAngle) * XAxis + FMath::Sin(CurrentAngle) * YAxis);
 			Vertices[2 * Side] = (MidVertex - HalfHeight) * CylinderScale;
 			Vertices[(2 * Side) + 1] = (MidVertex + HalfHeight) * CylinderScale;
 
@@ -254,7 +255,7 @@ void FLayer::BuildPokeAHoleMesh(TArray<FVector>& Vertices, TArray<int32>& Triang
 			}
 		}
 	}
-	else if (Desc.ShapeType == IStereoLayers::CubemapLayer)
+	else if (Desc.HasShape<FCubemapLayer>())
 	{
 		const float CubemapScale = 1000;
 		Vertices.Init(FVector::ZeroVector, 8);
@@ -371,25 +372,25 @@ void FLayer::Initialize_RenderThread(const FSettings* Settings, FCustomPresent*
 
 		ovrpShape Shape;
 
-		switch (Desc.ShapeType)
+		if (Desc.HasShape<FQuadLayer>())
 		{
-		case IStereoLayers::QuadLayer:
 			Shape = ovrpShape_Quad;
-			break;
-
-		case IStereoLayers::CylinderLayer:
+		}
+		else if (Desc.HasShape<FCylinderLayer>())
+		{
 			Shape = ovrpShape_Cylinder;
-			break;
-
-		case IStereoLayers::CubemapLayer:
+		}
+		else if (Desc.HasShape<FCubemapLayer>())
+		{
 			Shape = ovrpShape_Cubemap;
-			break;
-
-		case IStereoLayers::EquirectLayer:
+		}
+		else if (Desc.HasShape<FEquirectLayer>())
+		{
 			Shape = ovrpShape_Equirect;
-			break;
+		}
+		else
+		{
 
-		default:
 			return;
 		}
 
@@ -609,7 +610,6 @@ void FLayer::UpdateTexture_RenderThread(FCustomPresent* CustomPresent, FRHIComma
 		{
 			bool bAlphaPremultiply = true;
 			bool bNoAlphaWrite = (Desc.Flags & IStereoLayers::LAYER_FLAG_TEX_NO_ALPHA_CHANNEL) != 0;
-			bool bIsCubemap = (Desc.ShapeType == IStereoLayers::ELayerShape::CubemapLayer);
 
 			// Left
 			{
@@ -659,28 +659,30 @@ const ovrpLayerSubmit* FLayer::UpdateLayer_RHIThread(const FSettings* Settings,
 	OvrpLayerSubmit.ColorScale = injectColorScale ? Settings->ColorScale : ovrpVector4f{ 1, 1, 1, 1};
 
 	if (OvrpLayerDesc.Shape == ovrpShape_Equirect) {
+		const FEquirectLayer& EquirectProps = Desc.GetShape<FEquirectLayer>();
+
 		ovrpTextureRectMatrixf& RectMatrix = OvrpLayerSubmit.TextureRectMatrix;
 		ovrpRectf& LeftUVRect = RectMatrix.LeftRect;
 		ovrpRectf& RightUVRect = RectMatrix.RightRect;
-		LeftUVRect.Pos.x = Desc.EquirectProps.LeftUVRect.Min.X;
-		LeftUVRect.Pos.y = Desc.EquirectProps.LeftUVRect.Min.Y;
-		LeftUVRect.Size.w = Desc.EquirectProps.LeftUVRect.Max.X - Desc.EquirectProps.LeftUVRect.Min.X;
-		LeftUVRect.Size.h = Desc.EquirectProps.LeftUVRect.Max.Y - Desc.EquirectProps.LeftUVRect.Min.Y;
-		RightUVRect.Pos.x = Desc.EquirectProps.RightUVRect.Min.X;
-		RightUVRect.Pos.y = Desc.EquirectProps.RightUVRect.Min.Y;
-		RightUVRect.Size.w = Desc.EquirectProps.RightUVRect.Max.X - Desc.EquirectProps.RightUVRect.Min.X;
-		RightUVRect.Size.h = Desc.EquirectProps.RightUVRect.Max.Y - Desc.EquirectProps.RightUVRect.Min.Y;
+		LeftUVRect.Pos.x = EquirectProps.LeftUVRect.Min.X;
+		LeftUVRect.Pos.y = EquirectProps.LeftUVRect.Min.Y;
+		LeftUVRect.Size.w = EquirectProps.LeftUVRect.Max.X - EquirectProps.LeftUVRect.Min.X;
+		LeftUVRect.Size.h = EquirectProps.LeftUVRect.Max.Y - EquirectProps.LeftUVRect.Min.Y;
+		RightUVRect.Pos.x = EquirectProps.RightUVRect.Min.X;
+		RightUVRect.Pos.y = EquirectProps.RightUVRect.Min.Y;
+		RightUVRect.Size.w = EquirectProps.RightUVRect.Max.X - EquirectProps.RightUVRect.Min.X;
+		RightUVRect.Size.h = EquirectProps.RightUVRect.Max.Y - EquirectProps.RightUVRect.Min.Y;
 
 		ovrpVector4f& LeftScaleBias = RectMatrix.LeftScaleBias;
-		LeftScaleBias.x = Desc.EquirectProps.LeftScale.X;
-		LeftScaleBias.y = Desc.EquirectProps.LeftScale.Y;
-		LeftScaleBias.z = Desc.EquirectProps.LeftBias.X;
-		LeftScaleBias.w = Desc.EquirectProps.LeftBias.Y;
+		LeftScaleBias.x = EquirectProps.LeftScale.X;
+		LeftScaleBias.y = EquirectProps.LeftScale.Y;
+		LeftScaleBias.z = EquirectProps.LeftBias.X;
+		LeftScaleBias.w = EquirectProps.LeftBias.Y;
 		ovrpVector4f& RightScaleBias = RectMatrix.RightScaleBias;
-		RightScaleBias.x = Desc.EquirectProps.RightScale.X;
-		RightScaleBias.y = Desc.EquirectProps.RightScale.Y;
-		RightScaleBias.z = Desc.EquirectProps.RightBias.X;
-		RightScaleBias.w = Desc.EquirectProps.RightBias.Y;
+		RightScaleBias.x = EquirectProps.RightScale.X;
+		RightScaleBias.y = EquirectProps.RightScale.Y;
+		RightScaleBias.z = EquirectProps.RightBias.X;
+		RightScaleBias.w = EquirectProps.RightBias.Y;
 
 		OvrpLayerSubmit.OverrideTextureRectMatrix = ovrpBool_True;
 	}
@@ -705,10 +707,11 @@ const ovrpLayerSubmit* FLayer::UpdateLayer_RHIThread(const FSettings* Settings,
 			break;
 		case ovrpShape_Cylinder:
 			{
-				float CylinderHeight = (Desc.Flags & IStereoLayers::LAYER_FLAG_QUAD_PRESERVE_TEX_RATIO) ? Desc.CylinderOverlayArc * AspectRatio : Desc.CylinderHeight;
-				OvrpLayerSubmit.Cylinder.ArcWidth = Desc.CylinderOverlayArc * Scale.x;
+				const FCylinderLayer& CylinderProps = Desc.GetShape<FCylinderLayer>();
+				float CylinderHeight = (Desc.Flags & IStereoLayers::LAYER_FLAG_QUAD_PRESERVE_TEX_RATIO) ? CylinderProps.OverlayArc * AspectRatio : CylinderProps.Height;
+				OvrpLayerSubmit.Cylinder.ArcWidth = CylinderProps.OverlayArc * Scale.x;
 				OvrpLayerSubmit.Cylinder.Height = CylinderHeight * Scale.x;
-				OvrpLayerSubmit.Cylinder.Radius = Desc.CylinderRadius * Scale.x;
+				OvrpLayerSubmit.Cylinder.Radius = CylinderProps.Radius * Scale.x;
 			}
 			break;
 		}
diff --git a/Engine/Plugins/Runtime/Oculus/OculusVR/Source/OculusHMD/Private/OculusHMD_Splash.cpp b/Engine/Plugins/Runtime/Oculus/OculusVR/Source/OculusHMD/Private/OculusHMD_Splash.cpp
--- a/Engine/Plugins/Runtime/Oculus/OculusVR/Source/OculusHMD/Private/OculusHMD_Splash.cpp
+++ b/Engine/Plugins/Runtime/Oculus/OculusVR/Source/OculusHMD/Private/OculusHMD_Splash.cpp
@@ -39,7 +39,6 @@ FSplash::FSplash(FOculusHMD* InOculusHMD) :
 		LayerDesc.QuadSize = FVector2D(0.01f, 0.01f);
 		LayerDesc.Priority = 0;
 		LayerDesc.PositionType = IStereoLayers::TrackerLocked;
-		LayerDesc.ShapeType = IStereoLayers::QuadLayer;
 		LayerDesc.Texture = GBlackTexture->TextureRHI;
 		BlackLayer = MakeShareable(new FLayer(NextLayerId++, LayerDesc));
 	}
@@ -50,7 +49,6 @@ FSplash::FSplash(FOculusHMD* InOculusHMD) :
 		LayerDesc.QuadSize = FVector2D(0.01f, 0.01f);
 		LayerDesc.Priority = 0;
 		LayerDesc.PositionType = IStereoLayers::TrackerLocked;
-		LayerDesc.ShapeType = IStereoLayers::QuadLayer;
 		LayerDesc.Texture = nullptr;
 		UELayer = MakeShareable(new FLayer(NextLayerId++, LayerDesc));
 	}
@@ -422,15 +420,18 @@ bool FSplash::GetSplash(unsigned InSplashLayerIndex, FOculusSplashDesc& OutDesc)
 
 IStereoLayers::FLayerDesc FSplash::StereoLayerDescFromOculusSplashDesc(FOculusSplashDesc OculusDesc)
 {
-	bool bIsCubemap = OculusDesc.LoadedTexture->GetTextureCube() != nullptr;
-
 	IStereoLayers::FLayerDesc LayerDesc;
+	if (OculusDesc.LoadedTexture->GetTextureCube() != nullptr)
+	{
+		LayerDesc.SetShape<FCubemapLayer>();
+	}
+	// else LayerDesc.Shape defaults to FQuadLayer
+
 	LayerDesc.Transform = OculusDesc.TransformInMeters * FTransform(OculusHMD->GetSplashRotation().Quaternion());
 	LayerDesc.QuadSize = OculusDesc.QuadSizeInMeters;
 	LayerDesc.UVRect = FBox2D(OculusDesc.TextureOffset, OculusDesc.TextureOffset + OculusDesc.TextureScale);
 	LayerDesc.Priority = INT32_MAX - (int32)(OculusDesc.TransformInMeters.GetTranslation().X * 1000.f);
 	LayerDesc.PositionType = IStereoLayers::TrackerLocked;
-	LayerDesc.ShapeType = bIsCubemap ? IStereoLayers::CubemapLayer : IStereoLayers::QuadLayer;
 	LayerDesc.Texture = OculusDesc.LoadedTexture;
 	LayerDesc.Flags = IStereoLayers::LAYER_FLAG_QUAD_PRESERVE_TEX_RATIO | (OculusDesc.bNoAlphaChannel ? IStereoLayers::LAYER_FLAG_TEX_NO_ALPHA_CHANNEL : 0) | (OculusDesc.bIsDynamic ? IStereoLayers::LAYER_FLAG_TEX_CONTINUOUS_UPDATE : 0);
 
diff --git a/Engine/Plugins/Runtime/Steam/SteamVR/Source/SteamVR/Private/SteamVRHMD.h b/Engine/Plugins/Runtime/Steam/SteamVR/Source/SteamVR/Private/SteamVRHMD.h
--- a/Engine/Plugins/Runtime/Steam/SteamVR/Source/SteamVR/Private/SteamVRHMD.h
+++ b/Engine/Plugins/Runtime/Steam/SteamVR/Source/SteamVR/Private/SteamVRHMD.h
@@ -178,8 +178,7 @@ public:
 	virtual bool AllocateDepthTexture(uint32 Index, uint32 SizeX, uint32 SizeY, uint8 Format, uint32 NumMips, uint32 Flags, uint32 TargetableTextureFlags, FTexture2DRHIRef& OutTargetableTexture, FTexture2DRHIRef& OutShaderResourceTexture, uint32 NumSamples = 1) override;
 
 	// IStereoLayers interface
-	// Create/Set/Get/Destroy inherited from TStereoLayerManager
-	virtual void GetAllocatedTexture(uint32 LayerId, FTextureRHIRef &Texture, FTextureRHIRef &LeftTexture) override;
+	// Create/Set/Get/Destroy and GetAllocatedTexture inherited from TStereoLayerManager
 	virtual bool ShouldCopyDebugLayersToSpectatorScreen() const override { return true; }
 
 	// ISceneViewExtension interface
diff --git a/Engine/Plugins/Runtime/Steam/SteamVR/Source/SteamVR/Private/SteamVRStereoLayers.cpp b/Engine/Plugins/Runtime/Steam/SteamVR/Source/SteamVR/Private/SteamVRStereoLayers.cpp
--- a/Engine/Plugins/Runtime/Steam/SteamVR/Source/SteamVR/Private/SteamVRStereoLayers.cpp
+++ b/Engine/Plugins/Runtime/Steam/SteamVR/Source/SteamVR/Private/SteamVRStereoLayers.cpp
@@ -95,7 +95,7 @@ void FSteamVRHMD::UpdateLayer(struct FSteamVRLayer& Layer, uint32 LayerId, bool
 	{
 		UE_LOG(LogHMD, Warning, TEXT("Unsupported StereoLayer flag. SteamVR StereoLayers do not support disabling alpha renderding. Make the texture opaque instead."));
 	}
-	if (Layer.LayerDesc.ShapeType != IStereoLayers::QuadLayer)
+	if (!Layer.LayerDesc.HasShape<FQuadLayer>())
 	{
 		UE_LOG(LogHMD, Warning, TEXT("Unsupported StereoLayer shape. SteamVR StereoLayers can only be Quads."));
 	}
@@ -280,36 +280,6 @@ void FSteamVRHMD::UpdateStereoLayers_RenderThread()
 	}
 }
 
-void FSteamVRHMD::GetAllocatedTexture(uint32 LayerId, FTextureRHIRef &Texture, FTextureRHIRef &LeftTexture)
-{
-	Texture = LeftTexture = nullptr;
-	FSteamVRLayer* LayerFound = nullptr;
-	check(IsInRenderingThread()); // Not strictly necessary, as WithLayer uses a scope lock
-
-	WithLayer(LayerId, [&](FSteamVRLayer* LayerFound)
-	{
-		if (LayerFound && LayerFound->LayerDesc.Texture)
-		{
-			switch (LayerFound->LayerDesc.ShapeType)
-			{
-			case IStereoLayers::CubemapLayer:
-				Texture = LayerFound->LayerDesc.Texture->GetTextureCube();
-				LeftTexture = LayerFound->LayerDesc.LeftTexture ? LayerFound->LayerDesc.LeftTexture->GetTextureCube() : nullptr;
-				break;
-
-			case IStereoLayers::CylinderLayer:
-			case IStereoLayers::QuadLayer:
-				Texture = LayerFound->LayerDesc.Texture->GetTexture2D();
-				LeftTexture = LayerFound->LayerDesc.LeftTexture ? LayerFound->LayerDesc.LeftTexture->GetTexture2D() : nullptr;
-				break;
-
-			default:
-				break;
-			}
-		}
-	});
-}
-
 //=============================================================================
 IStereoLayers* FSteamVRHMD::GetStereoLayers()
 {
diff --git a/Engine/Source/Runtime/Engine/Private/Components/StereoLayerComponent.cpp b/Engine/Source/Runtime/Engine/Private/Components/StereoLayerComponent.cpp
--- a/Engine/Source/Runtime/Engine/Private/Components/StereoLayerComponent.cpp
+++ b/Engine/Source/Runtime/Engine/Private/Components/StereoLayerComponent.cpp
@@ -6,6 +6,7 @@
 #include "TextureResource.h"
 #include "Engine/Texture.h"
 #include "IStereoLayers.h"
+#include "StereoLayerShapes.h"
 #include "StereoRendering.h"
 
 UStereoLayerComponent::UStereoLayerComponent(const FObjectInitializer& ObjectInitializer)
@@ -119,23 +120,18 @@ void UStereoLayerComponent::TickComponent(float DeltaTime, enum ELevelTick TickT
 		switch (StereoLayerShape)
 		{
 		case SLSH_QuadLayer:
-			LayerDsec.ShapeType = IStereoLayers::QuadLayer;
 			break;
 
 		case SLSH_CylinderLayer:
-			LayerDsec.ShapeType = IStereoLayers::CylinderLayer;
-			LayerDsec.CylinderRadius = CylinderRadius;
-			LayerDsec.CylinderOverlayArc = CylinderOverlayArc;
-			LayerDsec.CylinderHeight = CylinderHeight;
+			LayerDsec.SetShape<FCylinderLayer>(CylinderRadius, CylinderOverlayArc, CylinderHeight);
 			break;
 
 		case SLSH_CubemapLayer:
-			LayerDsec.ShapeType = IStereoLayers::CubemapLayer;
+			LayerDsec.SetShape<FCubemapLayer>();
 			break;
 
 		case SLSH_EquirectLayer:
-			LayerDsec.ShapeType = IStereoLayers::EquirectLayer;
-			LayerDsec.EquirectProps = { EquirectProps.LeftUVRect, EquirectProps.RightUVRect, EquirectProps.LeftScale, EquirectProps.RightScale, EquirectProps.LeftBias, EquirectProps.RightBias };
+			LayerDsec.SetShape<FEquirectLayer>(EquirectProps.LeftUVRect, EquirectProps.RightUVRect, EquirectProps.LeftScale, EquirectProps.RightScale, EquirectProps.LeftBias, EquirectProps.RightBias );
 			break;
 
 		default:
diff --git a/Engine/Source/Runtime/Engine/Private/StereoLayerShapes.cpp b/Engine/Source/Runtime/Engine/Private/StereoLayerShapes.cpp
new file mode 100644
--- /dev/null
+++ b/Engine/Source/Runtime/Engine/Private/StereoLayerShapes.cpp
@@ -0,0 +1,9 @@
+// Copyright 1998-2019 Epic Games, Inc. All Rights Reserved.
+
+#include "StereoLayerShapes.h"
+#include "RHI.h"
+
+const FName FQuadLayer::ShapeName = FName("QuadLayer");
+const FName FCylinderLayer::ShapeName = FName("CylinderLayer");
+const FName FCubemapLayer::ShapeName = FName("CubemapLayer");
+const FName FEquirectLayer::ShapeName = FName("EquirectLayer");
diff --git a/Engine/Source/Runtime/Engine/Public/IStereoLayers.h b/Engine/Source/Runtime/Engine/Public/IStereoLayers.h
--- a/Engine/Source/Runtime/Engine/Public/IStereoLayers.h
+++ b/Engine/Source/Runtime/Engine/Public/IStereoLayers.h
@@ -8,10 +8,50 @@
 #pragma once
 
 #include "CoreMinimal.h"
+#include "StereoLayerShapes.h"
 #include "RHI.h"
 
 class IStereoLayers
 {
+
+	// The pointer to the shape of a stereo layer is wrapped in this simple class that clones and destroys the shape.
+	// This is to avoid having to implement non-default copy and assignment operations for the entire FLayerDesc struct.
+	struct FShapeWrapper
+	{
+		TUniquePtr<IStereoLayerShape> Wrapped;
+		FShapeWrapper(IStereoLayerShape* InWrapped)
+			: Wrapped(InWrapped)
+		{}
+
+		FShapeWrapper(const FShapeWrapper& In) :
+			Wrapped(In.IsValid()?In.Wrapped->Clone():nullptr)
+		{}
+
+		FShapeWrapper(FShapeWrapper&& In)
+		{
+			Wrapped = MoveTemp(In.Wrapped);
+		}
+
+		FShapeWrapper& operator= (IStereoLayerShape* InWrapped)
+		{
+			Wrapped = TUniquePtr<IStereoLayerShape>(InWrapped);
+			return *this;
+		}
+
+		FShapeWrapper& operator= (const FShapeWrapper& In)
+		{
+			return operator=(In.IsValid() ? In.Wrapped->Clone() : nullptr);
+		}
+
+		FShapeWrapper& operator= (FShapeWrapper&& In)
+		{
+			Wrapped = MoveTemp(In.Wrapped);
+			return *this;
+		}
+
+		bool IsValid() const { return Wrapped.IsValid(); }
+	};
+
 public:
 
 	enum ELayerType
@@ -21,14 +61,6 @@ public:
 		FaceLocked
 	};
 
-	enum ELayerShape
-	{
-		QuadLayer,
-		CylinderLayer,
-		CubemapLayer,
-		EquirectLayer
-	};
-
 	enum ELayerFlags
 	{
 		// Internally copies the texture on every frame for video, etc.
@@ -45,33 +77,23 @@ public:
 		LAYER_FLAG_HIDDEN = 0x00000020,
 	};
 
-	/** Structure describing additional settings for equirect layers */
-	struct FEquirectProps
-	{
-		/** Left source texture UVRect, specifying portion of input texture corresponding to left eye. */
-		FBox2D LeftUVRect;
-
-		/** Right source texture UVRect, specifying portion of input texture corresponding to right eye. */
-		FBox2D RightUVRect;
-
-		/** Left eye's texture coordinate scale after mapping to 2D. */
-		FVector2D LeftScale;
-
-		/** Right eye's texture coordinate scale after mapping to 2D. */
-		FVector2D RightScale;
-
-		/** Left eye's texture coordinate bias after mapping to 2D. */
-		FVector2D LeftBias;
-
-		/** Right eye's texture coordinate bias after mapping to 2D. */
-		FVector2D RightBias;
-	};
 
 	/**
 	 * Structure describing the visual appearance of a single stereo layer
 	 */
 	struct FLayerDesc
 	{
+
+		FLayerDesc()
+			: Shape(new FQuadLayer())
+		{
+		}
+
+		FLayerDesc(const IStereoLayerShape& InShape)
+			: Shape(InShape.Clone())
+		{
+		}
+
 		void SetLayerId(uint32 InId) { Id = InId; }
 		uint32 GetLayerId() const { return Id; }
 		bool IsVisible() const { return Texture != nullptr && !(Flags & LAYER_FLAG_HIDDEN); }
@@ -93,15 +115,14 @@ public:
 		int32				Priority	 = 0;
 		// Which space the layer is locked within
 		ELayerType			PositionType = ELayerType::FaceLocked;
-		// which shape of layer it is. ELayerShape::QuadLayer is the only shape supported by all VR platforms.
-		ELayerShape			ShapeType	 = ELayerShape::QuadLayer;
+		// which shape of layer it is. FQuadLayer is the only shape supported by all VR platforms.
+		// queries the shape of the layer at run time
+		template <typename T> bool HasShape() const { check(Shape.IsValid()); return Shape.Wrapped->GetShapeName() == T::ShapeName; }
+		// returns Shape cast to the supplied type. It's up to the caller to have ensured the cast is valid before calling this method.
+		template <typename T> T& GetShape() { check(HasShape<T>()); return *static_cast<T*>(Shape.Wrapped.Get()); }
+		template <typename T> const T& GetShape() const { check(HasShape<T>()); return *static_cast<const T*>(Shape.Wrapped.Get()); }
+		template <typename T, typename... InArgTypes> void SetShape(InArgTypes&&... Args) { Shape = FShapeWrapper(new T(Forward<InArgTypes>(Args)...)); }
 
-		// UVs and Scale/Bias of Equirect Layers.
-		FEquirectProps		EquirectProps;
-		// Cylinder layer settings.
-		float				CylinderRadius;
-		float				CylinderOverlayArc;
-		float				CylinderHeight;
 
 		// Texture mapped for right eye (if one texture provided, mono assumed)
 		FTextureRHIRef		Texture		 = nullptr;	
@@ -109,6 +130,8 @@ public:
 		FTextureRHIRef		LeftTexture  = nullptr;
 		// Uses LAYER_FLAG_... -- See: ELayerFlags
 		uint32				Flags		 = 0;
+	private: 
+		FShapeWrapper		Shape;
 	};
 
 	virtual ~IStereoLayers() { }
@@ -281,7 +304,6 @@ public:
 		StereoLayerDesc.Transform = FTransform(FVector(100.f, 0, 0));
 		StereoLayerDesc.QuadSize = FVector2D(120.f, 120.f);
 		StereoLayerDesc.PositionType = IStereoLayers::ELayerType::FaceLocked;
-		StereoLayerDesc.ShapeType = IStereoLayers::ELayerShape::QuadLayer;
 		StereoLayerDesc.Texture = Texture;
 		StereoLayerDesc.Flags = IStereoLayers::ELayerFlags::LAYER_FLAG_TEX_CONTINUOUS_UPDATE;
 		StereoLayerDesc.Flags |= IStereoLayers::ELayerFlags::LAYER_FLAG_QUAD_PRESERVE_TEX_RATIO;
diff --git a/Engine/Source/Runtime/Engine/Public/StereoLayerShapes.h b/Engine/Source/Runtime/Engine/Public/StereoLayerShapes.h
new file mode 100644
--- /dev/null
+++ b/Engine/Source/Runtime/Engine/Public/StereoLayerShapes.h
@@ -0,0 +1,89 @@
+// Copyright 1998-2019 Epic Games, Inc. All Rights Reserved.
+
+#pragma once
+
+#include "CoreMinimal.h"
+/** Different layer shapes are defined as subclass of ILayerShape.
+  * Currently, only FQuadLayer is supported by all platforms, but plugins may define their own custom shapes.
+  * When encountering an unsupported shape type, implementations should treat the layer as a quad layer.
+  */
+
+#define STEREO_LAYER_SHAPE_BOILERPLATE(ClassName) \
+public: \
+	static const FName ShapeName; \
+	virtual FName GetShapeName() override { return ShapeName; } \
+	virtual IStereoLayerShape* Clone() const override { return new ClassName(*this); }
+
+class ENGINE_API IStereoLayerShape
+{
+public:
+	/** Shape name is used to identify the shape type. */
+	virtual FName GetShapeName() = 0;
+	virtual IStereoLayerShape* Clone() const = 0;
+	virtual ~IStereoLayerShape() {}
+};
+
+/** Quad layer is the default layer shape and contains no additional settings */
+class ENGINE_API FQuadLayer : public IStereoLayerShape
+{
+	STEREO_LAYER_SHAPE_BOILERPLATE(FQuadLayer)
+};
+
+/** Class describing additional settings for cylinder layers. Currently only supported by Oculus. */
+class ENGINE_API FCylinderLayer : public IStereoLayerShape
+{
+	STEREO_LAYER_SHAPE_BOILERPLATE(FCylinderLayer)
+public:
+
+	FCylinderLayer() {}
+	FCylinderLayer(float InRadius, float InOverlayArc, float InHeight) :
+		Radius(InRadius),
+		OverlayArc(InOverlayArc),
+		Height(InHeight)
+	{}	
+
+	float	Radius;
+	float	OverlayArc;
+	float	Height;
+};
+
+/** Class describing additional settings for cube map layers. Currently only supported by Oculus. */
+class ENGINE_API FCubemapLayer : public IStereoLayerShape
+{
+	STEREO_LAYER_SHAPE_BOILERPLATE(FCubemapLayer)
+};
+
+/** Class describing additional settings for equirect layers. Currently only supported by Oculus. */
+class ENGINE_API FEquirectLayer : public IStereoLayerShape
+{
+	STEREO_LAYER_SHAPE_BOILERPLATE(FEquirectLayer)
+public:
+
+	FEquirectLayer() {}
+	FEquirectLayer(FBox2D InLeftUVRect, FBox2D InRightUVRect, FVector2D InLeftScale, FVector2D InRightScale, FVector2D InLeftBias, FVector2D InRightBias) :
+		LeftUVRect(InLeftUVRect),
+		RightUVRect(InRightUVRect),
+		LeftScale(InLeftScale),
+		RightScale(InRightScale),
+		LeftBias(InLeftBias),
+		RightBias(InRightBias)
+	{}
+
+	/** Left source texture UVRect, specifying portion of input texture corresponding to left eye. */
+	FBox2D LeftUVRect;
+
+	/** Right source texture UVRect, specifying portion of input texture corresponding to right eye. */
+	FBox2D RightUVRect;
+
+	/** Left eye's texture coordinate scale after mapping to 2D. */
+	FVector2D LeftScale;
+
+	/** Right eye's texture coordinate scale after mapping to 2D. */
+	FVector2D RightScale;
+
+	/** Left eye's texture coordinate bias after mapping to 2D. */
+	FVector2D LeftBias;
+
+	/** Right eye's texture coordinate bias after mapping to 2D. */
+	FVector2D RightBias;
+};
diff --git a/Engine/Source/Runtime/HeadMountedDisplay/Private/DefaultStereoLayers.cpp b/Engine/Source/Runtime/HeadMountedDisplay/Private/DefaultStereoLayers.cpp
--- a/Engine/Source/Runtime/HeadMountedDisplay/Private/DefaultStereoLayers.cpp
+++ b/Engine/Source/Runtime/HeadMountedDisplay/Private/DefaultStereoLayers.cpp
@@ -290,48 +290,4 @@ void FDefaultStereoLayers::SetupViewFamily(FSceneViewFamily& InViewFamily)
 	FVector HmdPosition = FVector::ZeroVector;
 	HMDDevice->GetCurrentPose(IXRTrackingSystem::HMDDeviceId, HmdOrientation, HmdPosition);
 	HmdTransform = FTransform(HmdOrientation, HmdPosition);
-}
-
-void FDefaultStereoLayers::GetAllocatedTexture(uint32 LayerId, FTextureRHIRef &Texture, FTextureRHIRef &LeftTexture)
-{
-	Texture = LeftTexture = nullptr;
-	FLayerDesc* LayerFound = nullptr;
-
-	if (IsInRenderingThread())
-	{
-		for (int32 LayerIndex = 0; LayerIndex < RenderThreadLayers.Num(); LayerIndex++)
-		{
-			if (RenderThreadLayers[LayerIndex].GetLayerId() == LayerId)
-			{
-				LayerFound = &RenderThreadLayers[LayerIndex];
-			}
-		}
-	}
-
-	else
-	{
-		// Only supporting the use of this function on RenderingThread.
-		check(false);
-		return;
-	}
-
-	if (LayerFound && LayerFound->Texture)
-	{
-		switch (LayerFound->ShapeType)
-		{
-		case IStereoLayers::CubemapLayer:
-			Texture = LayerFound->Texture->GetTextureCube();
-			LeftTexture = LayerFound->LeftTexture ? LayerFound->LeftTexture->GetTextureCube() : nullptr;			
-			break;
-
-		case IStereoLayers::CylinderLayer:
-		case IStereoLayers::QuadLayer:
-			Texture = LayerFound->Texture->GetTexture2D();
-			LeftTexture = LayerFound->LeftTexture ? LayerFound->LeftTexture->GetTexture2D() : nullptr;
-			break;
-
-		default:
-			break;
-		}
-	}
-}
+}
\ No newline at end of file
diff --git a/Engine/Source/Runtime/HeadMountedDisplay/Private/DefaultXRLoadingScreen.cpp b/Engine/Source/Runtime/HeadMountedDisplay/Private/DefaultXRLoadingScreen.cpp
--- a/Engine/Source/Runtime/HeadMountedDisplay/Private/DefaultXRLoadingScreen.cpp
+++ b/Engine/Source/Runtime/HeadMountedDisplay/Private/DefaultXRLoadingScreen.cpp
@@ -86,7 +86,22 @@ void FDefaultXRLoadingScreen::DoShowSplash(FSplashData& Splash)
 	// Sort layers by distance from the viewer
 	LayerDesc.Priority = INT32_MAX - static_cast<int32>(Splash.Desc.Transform.GetTranslation().X * 1000.f);
 	LayerDesc.PositionType = IStereoLayers::TrackerLocked;
-	LayerDesc.ShapeType = Splash.Desc.Texture->GetTextureCube() != nullptr ? IStereoLayers::CubemapLayer : IStereoLayers::QuadLayer;
+	
+	if (Splash.Desc.Texture->GetTextureCube() != nullptr)
+	{
+		if (!LayerDesc.HasShape<FCubemapLayer>())
+		{
+			LayerDesc.SetShape<FCubemapLayer>();
+		}
+	}
+	else
+	{
+		if (!LayerDesc.HasShape<FQuadLayer>())
+		{
+			LayerDesc.SetShape<FQuadLayer>();
+		}
+	}
+
 	LayerDesc.Texture = Splash.Desc.Texture;
 	LayerDesc.LeftTexture = Splash.Desc.LeftTexture;
 	LayerDesc.Flags = IStereoLayers::LAYER_FLAG_QUAD_PRESERVE_TEX_RATIO
diff --git a/Engine/Source/Runtime/HeadMountedDisplay/Public/DefaultStereoLayers.h b/Engine/Source/Runtime/HeadMountedDisplay/Public/DefaultStereoLayers.h
--- a/Engine/Source/Runtime/HeadMountedDisplay/Public/DefaultStereoLayers.h
+++ b/Engine/Source/Runtime/HeadMountedDisplay/Public/DefaultStereoLayers.h
@@ -24,8 +24,6 @@ public:
 	virtual void PostRenderView_RenderThread(FRHICommandListImmediate& RHICmdList, FSceneView& InView) override;
 	virtual bool IsActiveThisFrame(class FViewport* InViewport) const override;
 
-	void GetAllocatedTexture(uint32 LayerId, FTextureRHIRef &Texture, FTextureRHIRef &LeftTexture);
-
 	virtual bool ShouldCopyDebugLayersToSpectatorScreen() const override
 	{
 		// Emulated layer support means that the debug layer will be in the 3d scene render that the spectator screen displays.
diff --git a/Engine/Source/Runtime/HeadMountedDisplay/Public/StereoLayerManager.h b/Engine/Source/Runtime/HeadMountedDisplay/Public/StereoLayerManager.h
--- a/Engine/Source/Runtime/HeadMountedDisplay/Public/StereoLayerManager.h
+++ b/Engine/Source/Runtime/HeadMountedDisplay/Public/StereoLayerManager.h
@@ -296,6 +296,37 @@ public:
 		bStereoLayersDirty = true;
 	}
 
+	virtual void GetAllocatedTexture(uint32 LayerId, FTextureRHIRef &Texture, FTextureRHIRef &LeftTexture) override
+	{
+		Texture = LeftTexture = nullptr;
+		if (LayerId == FLayerDesc::INVALID_LAYER_ID)
+		{
+			return;
+		}
+		FScopeLock LockLayers(&LayerCritSect);
+
+		int32 FoundLevel;
+		LayerType* Found = FindLayerById(LayerId, FoundLevel);
+		if (Found)
+		{
+			FLayerDesc LayerDesc;
+			GetLayerDescMember(*Found, LayerDesc);
+			if (LayerDesc.Texture)
+			{
+				if (LayerDesc.HasShape<FCubemapLayer>())
+				{
+					Texture = LayerDesc.Texture->GetTextureCube();
+					LeftTexture = LayerDesc.LeftTexture ? LayerDesc.LeftTexture->GetTextureCube() : nullptr;
+				}
+				else
+				{
+					Texture = LayerDesc.Texture->GetTexture2D();
+					LeftTexture = LayerDesc.LeftTexture ? LayerDesc.LeftTexture->GetTexture2D() : nullptr;
+				}
+			}
+		}
+	}
+
 	virtual bool SupportsLayerState() override { return true; }
 
 	virtual void HideBackgroundLayer() { LayerState().bShowBackground = false; }
