Fix compilation on Android after Stereo Layer refactoring

diff --git a/Engine/Plugins/Runtime/Oculus/OculusVR/Source/OculusHMD/Private/OculusHMD_Layer.cpp b/Engine/Plugins/Runtime/Oculus/OculusVR/Source/OculusHMD/Private/OculusHMD_Layer.cpp
--- a/Engine/Plugins/Runtime/Oculus/OculusVR/Source/OculusHMD/Private/OculusHMD_Layer.cpp
+++ b/Engine/Plugins/Runtime/Oculus/OculusVR/Source/OculusHMD/Private/OculusHMD_Layer.cpp
@@ -217,13 +217,13 @@ void FLayer::BuildPokeAHoleMesh(TArray<FVector>& Vertices, TArray<int32>& Triang
 		FIntPoint TexSize = Desc.Texture.IsValid() ? Desc.Texture->GetTexture2D()->GetSizeXY() : Desc.LayerSize;
 		float AspectRatio = TexSize.X ? (float)TexSize.Y / (float)TexSize.X : 3.0f / 4.0f;
 
-		float CylinderHeight = (Desc.Flags & IStereoLayers::LAYER_FLAG_QUAD_PRESERVE_TEX_RATIO) ? CylinderProps->OverlayArc * AspectRatio : CylinderProps->Height;
+		float CylinderHeight = (Desc.Flags & IStereoLayers::LAYER_FLAG_QUAD_PRESERVE_TEX_RATIO) ? CylinderProps.OverlayArc * AspectRatio : CylinderProps.Height;
 
 		const FVector XAxis = FVector(1, 0, 0);
 		const FVector YAxis = FVector(0, 1, 0);
 		const FVector HalfHeight = FVector(0, 0, CylinderHeight / 2);
 
-		const float ArcAngle = CylinderProps->OverlayArc / CylinderProps->Radius;
+		const float ArcAngle = CylinderProps.OverlayArc / CylinderProps.Radius;
 		const int Sides = (int)( (ArcAngle * 180) / (PI * 5) ); // one triangle every 10 degrees of cylinder for a good-cheap approximation
 		Vertices.Init(FVector::ZeroVector, 2 * (Sides + 1));
 		UV0.Init(FVector2D::ZeroVector, 2 * (Sides + 1));
@@ -235,7 +235,7 @@ void FLayer::BuildPokeAHoleMesh(TArray<FVector>& Vertices, TArray<int32>& Triang
 
 		for (int Side = 0; Side < Sides + 1; Side++)
 		{
-			FVector MidVertex = CylinderProps->Radius * (FMath::Cos(CurrentAngle) * XAxis + FMath::Sin(CurrentAngle) * YAxis);
+			FVector MidVertex = CylinderProps.Radius * (FMath::Cos(CurrentAngle) * XAxis + FMath::Sin(CurrentAngle) * YAxis);
 			Vertices[2 * Side] = (MidVertex - HalfHeight) * CylinderScale;
 			Vertices[(2 * Side) + 1] = (MidVertex + HalfHeight) * CylinderScale;
 
