<saved by Perforce>

diff --git a/Engine/Plugins/Runtime/Oculus/OculusVR/Source/OculusHMD/Private/OculusHMD.cpp b/Engine/Plugins/Runtime/Oculus/OculusVR/Source/OculusHMD/Private/OculusHMD.cpp
--- a/Engine/Plugins/Runtime/Oculus/OculusVR/Source/OculusHMD/Private/OculusHMD.cpp
+++ b/Engine/Plugins/Runtime/Oculus/OculusVR/Source/OculusHMD/Private/OculusHMD.cpp
@@ -1420,7 +1420,6 @@ namespace OculusHMD
 
 	bool FOculusHMD::ShouldUseSeparateRenderTarget() const
 	{
-		CheckInGameThread();
 		return IsStereoEnabled();
 	}
 
