Use the new Oculus Touch keys for Menu and System buttons instead of the Gamepad special left and right

diff --git a/Engine/Plugins/Runtime/Oculus/OculusVR/Source/OculusInput/Private/OculusInputState.h b/Engine/Plugins/Runtime/Oculus/OculusVR/Source/OculusInput/Private/OculusInputState.h
--- a/Engine/Plugins/Runtime/Oculus/OculusVR/Source/OculusInput/Private/OculusInputState.h
+++ b/Engine/Plugins/Runtime/Oculus/OculusVR/Source/OculusInput/Private/OculusInputState.h
@@ -309,7 +309,7 @@ struct FOculusTouchControllerState
 		Buttons[ (int32)EOculusTouchControllerButton::Thumbstick_Left ].Key = (Hand == EControllerHand::Left) ? EKeys::OculusTouch_Left_Thumbstick_Left.GetFName() : EKeys::OculusTouch_Right_Thumbstick_Left.GetFName();
 		Buttons[ (int32)EOculusTouchControllerButton::Thumbstick_Right ].Key = (Hand == EControllerHand::Left) ? EKeys::OculusTouch_Left_Thumbstick_Right.GetFName() : EKeys::OculusTouch_Right_Thumbstick_Right.GetFName();
 
-		Buttons[ (int32)EOculusTouchControllerButton::Menu ].Key = (Hand == EControllerHand::Left) ? FGamepadKeyNames::SpecialLeft : FGamepadKeyNames::SpecialRight;
+		Buttons[(int32)EOculusTouchControllerButton::Menu].Key = (Hand == EControllerHand::Left) ? EKeys::OculusTouch_Left_Menu_Click.GetFName() : EKeys::OculusTouch_Right_System_Click.GetFName();
 
 		Buttons[ (int32)EOculusTouchControllerButton::Thumbstick_Touch ].Key = (Hand == EControllerHand::Left) ? EKeys::OculusTouch_Left_Thumbstick_Touch.GetFName() : EKeys::OculusTouch_Right_Thumbstick_Touch.GetFName();
 		Buttons[ (int32)EOculusTouchControllerButton::Trigger_Touch ].Key = (Hand == EControllerHand::Left) ? EKeys::OculusTouch_Left_Trigger_Touch.GetFName() : EKeys::OculusTouch_Right_Trigger_Touch.GetFName();
