Fix SingleEyeLetterboxed mode in Spectator Screen for Oculus.
Issue was Oculus was setting the viewport with

    RHICmdList.SetViewport(DstRect.Min.X, DstRect.Min.Y, 0.0f, MipViewportWidth, MipViewportHeight, 1.0f);

Which is incorrect since the 4th and 5th parameters are supposed to be max X and Y (not width and height); this worked in most modes since Min.X and Min.Y are usually 0. In this mode, Min.X is non-zero which exposes the bug.

diff --git a/Engine/Plugins/Runtime/Oculus/OculusVR/Source/OculusHMD/Private/OculusHMD_CustomPresent.cpp b/Engine/Plugins/Runtime/Oculus/OculusVR/Source/OculusHMD/Private/OculusHMD_CustomPresent.cpp
--- a/Engine/Plugins/Runtime/Oculus/OculusVR/Source/OculusHMD/Private/OculusHMD_CustomPresent.cpp
+++ b/Engine/Plugins/Runtime/Oculus/OculusVR/Source/OculusHMD/Private/OculusHMD_CustomPresent.cpp
@@ -438,7 +438,7 @@ void FCustomPresent::CopyTexture_RenderThread(FRHICommandListImmediate& RHICmdLi
 					PixelShader->SetParameters(RHICmdList, SamplerState, SrcTextureRHI, MipIndex);
 				}
 
-				RHICmdList.SetViewport(DstRect.Min.X, DstRect.Min.Y, 0.0f, MipViewportWidth, MipViewportHeight, 1.0f);
+				RHICmdList.SetViewport(DstRect.Min.X, DstRect.Min.Y, 0.0f, DstRect.Min.X + MipViewportWidth, DstRect.Min.Y + MipViewportHeight, 1.0f);
 
 				RendererModule->DrawRectangle(
 					RHICmdList,
