Settings naming fixup from Oculus

diff --git a/Engine/Plugins/Runtime/Oculus/OculusVR/Source/OculusEditor/Private/OculusToolWidget.cpp b/Engine/Plugins/Runtime/Oculus/OculusVR/Source/OculusEditor/Private/OculusToolWidget.cpp
--- a/Engine/Plugins/Runtime/Oculus/OculusVR/Source/OculusEditor/Private/OculusToolWidget.cpp
+++ b/Engine/Plugins/Runtime/Oculus/OculusVR/Source/OculusEditor/Private/OculusToolWidget.cpp
@@ -219,9 +219,9 @@ void SOculusToolWidget::RebuildLayout()
 	AddSimpleSetting(box, SimpleSettings.Find(FName("SupportDash")));
 	AddSimpleSetting(box, SimpleSettings.Find(FName("ForwardShading")));
 	AddSimpleSetting(box, SimpleSettings.Find(FName("AllowStaticLighting")));
-	AddSimpleSetting(box, SimpleSettings.Find(FName("InstancedStereo")));
+	AddSimpleSetting(box, SimpleSettings.Find(FName("MultiView")));
 	AddSimpleSetting(box, SimpleSettings.Find(FName("MobileMultiView")));
-	AddSimpleSetting(box, SimpleSettings.Find(FName("MobileHDR")));
+	AddSimpleSetting(box, SimpleSettings.Find(FName("MobilePostProcessing")));
 	AddSimpleSetting(box, SimpleSettings.Find(FName("AndroidManifest")));
 	AddSimpleSetting(box, SimpleSettings.Find(FName("AndroidPackaging")));
 
