Rename settings back to original names in the UI

Awaiting the relevant feature changes to be fully implemented before renaming the settings

diff --git a/Engine/Plugins/Runtime/Oculus/OculusVR/Source/OculusEditor/Private/OculusToolWidget.cpp b/Engine/Plugins/Runtime/Oculus/OculusVR/Source/OculusEditor/Private/OculusToolWidget.cpp
--- a/Engine/Plugins/Runtime/Oculus/OculusVR/Source/OculusEditor/Private/OculusToolWidget.cpp
+++ b/Engine/Plugins/Runtime/Oculus/OculusVR/Source/OculusEditor/Private/OculusToolWidget.cpp
@@ -365,13 +365,13 @@ void SOculusToolWidget::Construct(const FArguments& InArgs)
 
 	SimpleSettings.Add(FName("MultiView"), {
 		FName("MultiView"),
-		LOCTEXT("MultiViewDescription", "Multi-view is not enabled for this project. Instanced stereo substantially reduces draw calls, and improves rendering performance."),
+		LOCTEXT("InstancedStereoDescription", "Instanced stereo is not enabled for this project. Instanced stereo substantially reduces draw calls, and improves rendering performance."),
 		&SOculusToolWidget::MultiViewVisibility,
 		TArray<SimpleSettingAction>(),
 		(int)SupportFlags::SupportPC
 	});
 	SimpleSettings.Find(FName("MultiView"))->actions.Add(
-		{ LOCTEXT("MultiViewButtonText", "Enable Multi-View"),
+		{ LOCTEXT("InstancedStereoButtonText", "Enable Instanced Stereo"),
 		&SOculusToolWidget::MultiViewEnable }
 	);
 
@@ -389,13 +389,13 @@ void SOculusToolWidget::Construct(const FArguments& InArgs)
 
 	SimpleSettings.Add(FName("MobilePostProcessing"), {
 		FName("MobilePostProcessing"),
-		LOCTEXT("MobilePostProcessingDescription", "Mobile Post-Processing has performance and stability issues in VR. We strongly recommend disabling it."),
+		LOCTEXT("MobileHDRDescription", "Mobile HDR has performance and stability issues in VR. We strongly recommend disabling it."),
 		&SOculusToolWidget::MobilePostProcessingVisibility,
 		TArray<SimpleSettingAction>(),
 		(int)SupportFlags::SupportMobile
 	});
 	SimpleSettings.Find(FName("MobilePostProcessing"))->actions.Add(
-		{ LOCTEXT("MobileHDRButton", "Disable Mobile Post-Processing"),
+		{ LOCTEXT("MobileHDRButton", "Disable Mobile HDR"),
 		&SOculusToolWidget::MobilePostProcessingDisable }
 	);
 
diff --git a/Engine/Source/Runtime/Engine/Classes/Engine/RendererSettings.h b/Engine/Source/Runtime/Engine/Classes/Engine/RendererSettings.h
--- a/Engine/Source/Runtime/Engine/Classes/Engine/RendererSettings.h
+++ b/Engine/Source/Runtime/Engine/Classes/Engine/RendererSettings.h
@@ -605,7 +605,7 @@ class ENGINE_API URendererSettings : public UDeveloperSettings
 	uint32 bMultiView : 1;
 
 	UPROPERTY(config, EditAnywhere, Category = VR, meta=(
-		ConsoleVariable="r.MobileHDR",DisplayName="Mobile Post-Processing",
+		ConsoleVariable="r.MobileHDR", DisplayName="Mobile HDR",
 		ToolTip="If true, mobile pipelines include a full post-processing pass with tonemapping. Disable this setting for a performance boost and to enable stereoscopic rendering optimizations. Changing this setting requires restarting the editor.",
 		ConfigRestartRequired = true))
 	uint32 bMobilePostProcessing:1;
@@ -617,7 +617,7 @@ class ENGINE_API URendererSettings : public UDeveloperSettings
 		ConfigRestartRequired = true))
 		uint32 bMobileMultiView : 1;
 	
-	UPROPERTY(config, EditAnywhere, Category = VR, meta = (
+	UPROPERTY(config, meta = (
 		EditCondition = "!bMobilePostProcessing",
 		ConsoleVariable = "r.Mobile.UseHWsRGBEncoding", DisplayName = "Single-pass linear rendering",
 		ToolTip = "If true then mobile single-pass (without post-processing) rendering will use HW accelerated sRGB encoding/decoding. Available only on Oculus for now."))
