Make missing OVR dll warnings more friendly.

diff --git a/Engine/Plugins/Runtime/Oculus/OculusVR/Source/OculusHMD/OculusMobile_APL.xml b/Engine/Plugins/Runtime/Oculus/OculusVR/Source/OculusHMD/OculusMobile_APL.xml
--- a/Engine/Plugins/Runtime/Oculus/OculusVR/Source/OculusHMD/OculusMobile_APL.xml
+++ b/Engine/Plugins/Runtime/Oculus/OculusVR/Source/OculusHMD/OculusMobile_APL.xml
@@ -276,8 +276,8 @@
 		<!-- need this if plugin enabled and supported architecture, even if not packaged for Oculus mobile -->
 		<if condition="bSupported">
 			<true>
-				<loadLibrary name="vrapi" failmsg="vrapi library not loaded and required!" />
-        		<loadLibrary name="OVRPlugin" failmsg="OVRPlugin library not loaded and required!" />
+				<loadLibrary name="vrapi" failmsg="vrapi library not loaded and may be required for Oculus VR." />
+        		<loadLibrary name="OVRPlugin" failmsg="OVRPlugin library not loaded and may be required for Oculus VR." />
       		</true>
 		</if>
 	</soLoadLibrary>
