Fixed CI warning for uninitialized UObject property in FPropertyEditorTestInstancedStruct and FOculusCapsuleCollider.

diff --git a/Engine/Plugins/Runtime/Oculus/OculusVR/Source/OculusInput/Public/OculusInputFunctionLibrary.h b/Engine/Plugins/Runtime/Oculus/OculusVR/Source/OculusInput/Public/OculusInputFunctionLibrary.h
--- a/Engine/Plugins/Runtime/Oculus/OculusVR/Source/OculusInput/Public/OculusInputFunctionLibrary.h
+++ b/Engine/Plugins/Runtime/Oculus/OculusVR/Source/OculusInput/Public/OculusInputFunctionLibrary.h
@@ -74,7 +74,7 @@ struct OCULUSINPUT_API FOculusCapsuleCollider
 
 public:
 	UPROPERTY(BlueprintReadOnly, Category = "OculusLibrary|HandTracking")
-	UCapsuleComponent* Capsule;
+	UCapsuleComponent* Capsule { nullptr };
 
 	UPROPERTY(BlueprintReadOnly, Category = "OculusLibrary|HandTracking")
 	EBone BoneId;
diff --git a/Engine/Source/Editor/UnrealEd/Classes/Editor/PropertyEditorTestObject.h b/Engine/Source/Editor/UnrealEd/Classes/Editor/PropertyEditorTestObject.h
--- a/Engine/Source/Editor/UnrealEd/Classes/Editor/PropertyEditorTestObject.h
+++ b/Engine/Source/Editor/UnrealEd/Classes/Editor/PropertyEditorTestObject.h
@@ -178,7 +178,7 @@ struct FPropertyEditorTestInstancedStruct
 	GENERATED_BODY()
 
 	UPROPERTY(EditAnywhere, Instanced, Category=Default)
-	UPropertyEditorTestInstancedObject* Object;
+	UPropertyEditorTestInstancedObject* Object { nullptr };
 };
 
 UCLASS(transient)
