Kills FPostProcessSettings::ScreenPercentage

diff --git a/Engine/Plugins/Enterprise/VariantManagerContent/Source/VariantManagerContent/Private/ThumbnailGenerator.cpp b/Engine/Plugins/Enterprise/VariantManagerContent/Source/VariantManagerContent/Private/ThumbnailGenerator.cpp
--- a/Engine/Plugins/Enterprise/VariantManagerContent/Source/VariantManagerContent/Private/ThumbnailGenerator.cpp
+++ b/Engine/Plugins/Enterprise/VariantManagerContent/Source/VariantManagerContent/Private/ThumbnailGenerator.cpp
@@ -44,7 +44,7 @@ namespace ThumbnailGeneratorImpl
 		);
 
 		ViewFamily.SetScreenPercentageInterface(new FLegacyScreenPercentageDriver(
-			ViewFamily, /* GlobalResolutionFraction = */ 1.0f, /* AllowPostProcessSettingsScreenPercentage = */ false));
+			ViewFamily, /* GlobalResolutionFraction = */ 1.0f));
 
 		FSceneViewInitOptions ViewInitOptions;
 		ViewInitOptions.SetViewRectangle(FIntRect(0, 0, TargetSize.X, TargetSize.Y));
diff --git a/Engine/Plugins/Experimental/MeshModelingToolset/Source/ModelingComponents/Private/Scene/WorldRenderCapture.cpp b/Engine/Plugins/Experimental/MeshModelingToolset/Source/ModelingComponents/Private/Scene/WorldRenderCapture.cpp
--- a/Engine/Plugins/Experimental/MeshModelingToolset/Source/ModelingComponents/Private/Scene/WorldRenderCapture.cpp
+++ b/Engine/Plugins/Experimental/MeshModelingToolset/Source/ModelingComponents/Private/Scene/WorldRenderCapture.cpp
@@ -397,7 +397,7 @@ bool FWorldRenderCapture::CaptureMRSFromPosition(
 	ViewFamily.EngineShowFlags.SetContactShadows(false);
 
 	//ViewFamily.EngineShowFlags.SetScreenPercentage(false);
-	ViewFamily.SetScreenPercentageInterface(new FLegacyScreenPercentageDriver(ViewFamily, 1.f, false));
+	ViewFamily.SetScreenPercentageInterface(new FLegacyScreenPercentageDriver(ViewFamily, 1.f));
 
 	// This is called in various other places, unclear if we should be doing this too
 	//EngineShowFlagOverride(EShowFlagInitMode::ESFIM_Game, ViewModeIndex, ViewFamily.EngineShowFlags, true);
@@ -535,7 +535,7 @@ bool FWorldRenderCapture::CaptureEmissiveFromPosition(
 	NewView->CurrentBufferVisualizationMode = FName("PreTonemapHDRColor");
 	ViewFamily.Views.Add(NewView);
 
-	ViewFamily.SetScreenPercentageInterface(new FLegacyScreenPercentageDriver(ViewFamily, 1.0f, false));
+	ViewFamily.SetScreenPercentageInterface(new FLegacyScreenPercentageDriver(ViewFamily, 1.0f));
 
 	ViewFamily.EngineShowFlags.DisableAdvancedFeatures();
 	ViewFamily.EngineShowFlags.LOD = 0;
@@ -643,7 +643,7 @@ namespace Internal
 		NewView->CurrentBufferVisualizationMode = VisualizationMode;
 		ViewFamily.Views.Add(NewView);
 
-		ViewFamily.SetScreenPercentageInterface(new FLegacyScreenPercentageDriver(ViewFamily, 1.0f, false));
+		ViewFamily.SetScreenPercentageInterface(new FLegacyScreenPercentageDriver(ViewFamily, 1.0f));
 
 		// should we cache the FCanvas?
 		FCanvas Canvas(RenderTargetResource, NULL, FApp::GetCurrentTime() - GStartTime, FApp::GetDeltaTime(), FApp::GetCurrentTime() - GStartTime, Scene->GetFeatureLevel());
diff --git a/Engine/Plugins/FX/Niagara/Source/NiagaraEditor/Private/NiagaraBakerRenderer.cpp b/Engine/Plugins/FX/Niagara/Source/NiagaraEditor/Private/NiagaraBakerRenderer.cpp
--- a/Engine/Plugins/FX/Niagara/Source/NiagaraEditor/Private/NiagaraBakerRenderer.cpp
+++ b/Engine/Plugins/FX/Niagara/Source/NiagaraEditor/Private/NiagaraBakerRenderer.cpp
@@ -150,7 +150,7 @@ bool FNiagaraBakerRenderer::RenderView(UTextureRenderTarget2D* RenderTarget, FCa
 			NewView->CurrentBufferVisualizationMode = SourceName;
 			ViewFamily.Views.Add(NewView);
 
-			ViewFamily.SetScreenPercentageInterface(new FLegacyScreenPercentageDriver(ViewFamily, 1.0f, false));
+			ViewFamily.SetScreenPercentageInterface(new FLegacyScreenPercentageDriver(ViewFamily, 1.0f));
 
 			GetRendererModule().BeginRenderingViewFamily(Canvas, &ViewFamily);
 
diff --git a/Engine/Plugins/MovieScene/MovieRenderPipeline/Source/MovieRenderPipelineRenderPasses/Private/MoviePipelineImagePassBase.cpp b/Engine/Plugins/MovieScene/MovieRenderPipeline/Source/MovieRenderPipelineRenderPasses/Private/MoviePipelineImagePassBase.cpp
--- a/Engine/Plugins/MovieScene/MovieRenderPipeline/Source/MovieRenderPipelineRenderPasses/Private/MoviePipelineImagePassBase.cpp
+++ b/Engine/Plugins/MovieScene/MovieRenderPipeline/Source/MovieRenderPipelineRenderPasses/Private/MoviePipelineImagePassBase.cpp
@@ -95,7 +95,7 @@ TSharedPtr<FSceneViewFamilyContext> UMoviePipelineImagePassBase::CalculateViewFa
 	// If UMoviePipelineViewFamilySetting never set a Screen percentage interface we fallback to default.
 	if (OutViewFamily->GetScreenPercentageInterface() == nullptr)
 	{
-		OutViewFamily->SetScreenPercentageInterface(new FLegacyScreenPercentageDriver(*OutViewFamily, IsScreenPercentageSupported() ? InOutSampleState.GlobalScreenPercentageFraction : 1.f, IsScreenPercentageSupported()));
+		OutViewFamily->SetScreenPercentageInterface(new FLegacyScreenPercentageDriver(*OutViewFamily, IsScreenPercentageSupported() ? InOutSampleState.GlobalScreenPercentageFraction : 1.f));
 	}
 
 
@@ -507,4 +507,4 @@ void UMoviePipelineImagePassBase::BlendPostProcessSettings(FSceneView* InView)
 
 		InView->OverridePostProcessSettings(ViewInfo.PostProcessSettings, ViewInfo.PostProcessBlendWeight);
 	}
-}
+}
\ No newline at end of file
diff --git a/Engine/Plugins/Runtime/Oculus/OculusVR/Source/OculusHMD/Private/OculusHMD_DynamicResolutionState.cpp b/Engine/Plugins/Runtime/Oculus/OculusVR/Source/OculusHMD/Private/OculusHMD_DynamicResolutionState.cpp
--- a/Engine/Plugins/Runtime/Oculus/OculusVR/Source/OculusHMD/Private/OculusHMD_DynamicResolutionState.cpp
+++ b/Engine/Plugins/Runtime/Oculus/OculusVR/Source/OculusHMD/Private/OculusHMD_DynamicResolutionState.cpp
@@ -1,6 +1,7 @@
 // Copyright Epic Games, Inc. All Rights Reserved.
 
 #include "OculusHMD_DynamicResolutionState.h"
+#include "LegacyScreenPercentageDriver.h"
 
 #if OCULUS_HMD_SUPPORTED_PLATFORMS
 #include "SceneView.h"
@@ -8,62 +9,6 @@
 namespace OculusHMD
 {
 
-//-------------------------------------------------------------------------------------------------
-// FScreenPercentage
-//-------------------------------------------------------------------------------------------------
-
-class FScreenPercentage : public ISceneViewFamilyScreenPercentage
-{
-public:
-	FScreenPercentage(const FSceneViewFamily& InViewFamily, float InResolutionFraction, float InResolutionFractionUpperBound);
-
-protected:
-	// ISceneViewFamilyScreenPercentage
-	virtual float GetPrimaryResolutionFractionUpperBound() const override;
-	virtual ISceneViewFamilyScreenPercentage* Fork_GameThread(const class FSceneViewFamily& ForkedViewFamily) const override;
-	virtual void ComputePrimaryResolutionFractions_RenderThread(TArray<FSceneViewScreenPercentageConfig>& OutViewScreenPercentageConfigs) const override;
-
-private:
-	const FSceneViewFamily& ViewFamily;
-	const float ResolutionFraction;
-	const float ResolutionFractionUpperBound;
-};
-
-
-//-------------------------------------------------------------------------------------------------
-// FScreenPercentage implementation
-//-------------------------------------------------------------------------------------------------
-
-FScreenPercentage::FScreenPercentage(const FSceneViewFamily& InViewFamily, float InResolutionFraction, float InResolutionFractionUpperBound)
-	: ViewFamily(InViewFamily)
-	, ResolutionFraction(InResolutionFraction)
-	, ResolutionFractionUpperBound(InResolutionFractionUpperBound)
-{
-	check(ViewFamily.EngineShowFlags.ScreenPercentage == true);
-}
-
-float FScreenPercentage::GetPrimaryResolutionFractionUpperBound() const
-{
-	return ResolutionFractionUpperBound;
-}
-
-ISceneViewFamilyScreenPercentage* FScreenPercentage::Fork_GameThread(const class FSceneViewFamily& ForkedViewFamily) const
-{
-	return new FScreenPercentage(ForkedViewFamily, ResolutionFraction, ResolutionFractionUpperBound);
-}
-
-void FScreenPercentage::ComputePrimaryResolutionFractions_RenderThread(TArray<FSceneViewScreenPercentageConfig>& OutViewScreenPercentageConfigs) const
-{
-	check(IsInRenderingThread());
-	check(ViewFamily.EngineShowFlags.ScreenPercentage == true);
-
-	for (int32 ConfigIter = 0; ConfigIter < OutViewScreenPercentageConfigs.Num(); ++ConfigIter)
-	{
-		OutViewScreenPercentageConfigs[ConfigIter].PrimaryResolutionFraction = ResolutionFraction;
-	}
-}
-
-
 //-------------------------------------------------------------------------------------------------
 // FDynamicResolutionState implementation
 //-------------------------------------------------------------------------------------------------
@@ -102,24 +47,24 @@ void FDynamicResolutionState::SetupMainViewFamily(class FSceneViewFamily& ViewFa
 		float MaxResolutionFraction = Settings->PixelDensityMax;
 
 		// Clamp resolution fraction to what the renderer can do.
-		MinResolutionFraction = FMath::Max(MinResolutionFraction, FSceneViewScreenPercentageConfig::kMinResolutionFraction);
-		MaxResolutionFraction = FMath::Min(MaxResolutionFraction, FSceneViewScreenPercentageConfig::kMaxResolutionFraction);
+		MinResolutionFraction = FMath::Max(MinResolutionFraction, ISceneViewFamilyScreenPercentage::kMinResolutionFraction);
+		MaxResolutionFraction = FMath::Min(MaxResolutionFraction, ISceneViewFamilyScreenPercentage::kMaxResolutionFraction);
 
 		if (View.AntiAliasingMethod == AAM_TSR)
 		{
-			MinResolutionFraction = FMath::Max(MinResolutionFraction, FSceneViewScreenPercentageConfig::kMinTSRResolutionFraction);
-			MaxResolutionFraction = FMath::Min(MaxResolutionFraction, FSceneViewScreenPercentageConfig::kMaxTSRResolutionFraction);
+			MinResolutionFraction = FMath::Max(MinResolutionFraction, ISceneViewFamilyScreenPercentage::kMinTSRResolutionFraction);
+			MaxResolutionFraction = FMath::Min(MaxResolutionFraction, ISceneViewFamilyScreenPercentage::kMaxTSRResolutionFraction);
 		}
 		else if (View.AntiAliasingMethod == AAM_TemporalAA)
 		{
-			MinResolutionFraction = FMath::Max(MinResolutionFraction, FSceneViewScreenPercentageConfig::kMinTAAUpsampleResolutionFraction);
-			MaxResolutionFraction = FMath::Min(MaxResolutionFraction, FSceneViewScreenPercentageConfig::kMaxTAAUpsampleResolutionFraction);
+			MinResolutionFraction = FMath::Max(MinResolutionFraction, ISceneViewFamilyScreenPercentage::kMinTAAUpsampleResolutionFraction);
+			MaxResolutionFraction = FMath::Min(MaxResolutionFraction, ISceneViewFamilyScreenPercentage::kMaxTAAUpsampleResolutionFraction);
 		}
 
 		ResolutionFraction = FMath::Clamp(Settings->PixelDensity, MinResolutionFraction, MaxResolutionFraction);
 		ResolutionFractionUpperBound = MaxResolutionFraction;
 
-		ViewFamily.SetScreenPercentageInterface(new FScreenPercentage(ViewFamily, ResolutionFraction, ResolutionFractionUpperBound));
+		ViewFamily.SetScreenPercentageInterface(new FLegacyScreenPercentageDriver(ViewFamily, ResolutionFraction, ResolutionFractionUpperBound));
 	}
 }
 
diff --git a/Engine/Plugins/Runtime/nDisplay/Source/DisplayCluster/Private/Game/EngineClasses/Basics/DisplayClusterViewportClient.cpp b/Engine/Plugins/Runtime/nDisplay/Source/DisplayCluster/Private/Game/EngineClasses/Basics/DisplayClusterViewportClient.cpp
--- a/Engine/Plugins/Runtime/nDisplay/Source/DisplayCluster/Private/Game/EngineClasses/Basics/DisplayClusterViewportClient.cpp
+++ b/Engine/Plugins/Runtime/nDisplay/Source/DisplayCluster/Private/Game/EngineClasses/Basics/DisplayClusterViewportClient.cpp
@@ -474,7 +474,6 @@ void UDisplayClusterViewportClient::Draw(FViewport* InViewport, FCanvas* SceneCa
 						ViewFamily.SetScreenPercentageInterface(new FLegacyScreenPercentageDriver(
 							ViewFamily,
 							DynamicResolutionStateInfos.ResolutionFractionApproximation,
-							false,
 							DynamicResolutionStateInfos.ResolutionFractionUpperBound));
 					}
 
@@ -493,15 +492,11 @@ void UDisplayClusterViewportClient::Draw(FViewport* InViewport, FCanvas* SceneCa
 					// In case of stereo, we set the same buffer ratio for both left and right views (taken from left)
 					float CustomBufferRatio = DCViewFamily.CustomBufferRatio;
 
-					bool AllowPostProcessSettingsScreenPercentage = false;
 					float GlobalResolutionFraction = 1.0f;
 					float SecondaryScreenPercentage = 1.0f;
 
 					if (ViewFamily.EngineShowFlags.ScreenPercentage)
 					{
-						// Allow FPostProcessSettings::ScreenPercentage.
-						AllowPostProcessSettingsScreenPercentage = true;
-
 						// Get global view fraction set by r.ScreenPercentage.
 						GlobalResolutionFraction = FLegacyScreenPercentageDriver::GetCVarResolutionFraction() * CustomBufferRatio;
 
@@ -514,7 +509,7 @@ void UDisplayClusterViewportClient::Draw(FViewport* InViewport, FCanvas* SceneCa
 					}
 
 					ViewFamily.SetScreenPercentageInterface(new FLegacyScreenPercentageDriver(
-						ViewFamily, GlobalResolutionFraction, AllowPostProcessSettingsScreenPercentage));
+						ViewFamily, GlobalResolutionFraction));
 
 					ViewFamily.SecondaryViewFraction = SecondaryScreenPercentage;
 				}
diff --git a/Engine/Plugins/Runtime/nDisplay/Source/DisplayCluster/Private/Render/Viewport/DisplayClusterViewportManager_EditorPreview.cpp b/Engine/Plugins/Runtime/nDisplay/Source/DisplayCluster/Private/Render/Viewport/DisplayClusterViewportManager_EditorPreview.cpp
--- a/Engine/Plugins/Runtime/nDisplay/Source/DisplayCluster/Private/Render/Viewport/DisplayClusterViewportManager_EditorPreview.cpp
+++ b/Engine/Plugins/Runtime/nDisplay/Source/DisplayCluster/Private/Render/Viewport/DisplayClusterViewportManager_EditorPreview.cpp
@@ -164,7 +164,7 @@ bool FDisplayClusterViewportManager::RenderInEditor(class FDisplayClusterRenderF
 				{
 					// Screen percentage is still not supported in scene capture.
 					ViewFamily.EngineShowFlags.ScreenPercentage = false;
-					ViewFamily.SetScreenPercentageInterface(new FLegacyScreenPercentageDriver(ViewFamily, 1.0f, false));
+					ViewFamily.SetScreenPercentageInterface(new FLegacyScreenPercentageDriver(ViewFamily, 1.0f));
 
 					ViewFamily.bIsRenderedImmediatelyAfterAnotherViewFamily = bIsRenderedImmediatelyAfterAnotherViewFamily;
 
diff --git a/Engine/Source/Developer/MaterialUtilities/Private/MaterialUtilities.cpp b/Engine/Source/Developer/MaterialUtilities/Private/MaterialUtilities.cpp
--- a/Engine/Source/Developer/MaterialUtilities/Private/MaterialUtilities.cpp
+++ b/Engine/Source/Developer/MaterialUtilities/Private/MaterialUtilities.cpp
@@ -789,7 +789,7 @@ static void RenderSceneToTexture(
 	ViewFamily.Views.Add(NewView);
 
 	ViewFamily.SetScreenPercentageInterface(new FLegacyScreenPercentageDriver(
-		ViewFamily, /* GlobalResolutionFraction = */ 1.0f, /* AllowPostProcessSettingsScreenPercentage = */ false));
+		ViewFamily, /* GlobalResolutionFraction = */ 1.0f));
 
 	FCanvas Canvas(RenderTargetResource, NULL, FApp::GetCurrentTime() - GStartTime, FApp::GetDeltaTime(), FApp::GetCurrentTime() - GStartTime, Scene->GetFeatureLevel());
 	Canvas.Clear(FLinearColor::Transparent);
diff --git a/Engine/Source/Editor/LevelEditor/Private/SLevelViewportToolBar.cpp b/Engine/Source/Editor/LevelEditor/Private/SLevelViewportToolBar.cpp
--- a/Engine/Source/Editor/LevelEditor/Private/SLevelViewportToolBar.cpp
+++ b/Engine/Source/Editor/LevelEditor/Private/SLevelViewportToolBar.cpp
@@ -1333,8 +1333,8 @@ void SLevelViewportToolBar::OnFOVValueChanged( float NewValue )
 
 TSharedRef<SWidget> SLevelViewportToolBar::GenerateScreenPercentageMenu() const
 {
-	const int32 PreviewScreenPercentageMin = FSceneViewScreenPercentageConfig::kMinTAAUpsampleResolutionFraction * 100.0f;
-	const int32 PreviewScreenPercentageMax = FSceneViewScreenPercentageConfig::kMaxTAAUpsampleResolutionFraction * 100.0f;
+	const int32 PreviewScreenPercentageMin = ISceneViewFamilyScreenPercentage::kMinTAAUpsampleResolutionFraction * 100.0f;
+	const int32 PreviewScreenPercentageMax = ISceneViewFamilyScreenPercentage::kMaxTAAUpsampleResolutionFraction * 100.0f;
 
 	return
 		SNew(SBox)
diff --git a/Engine/Source/Editor/MovieSceneTools/Private/TrackEditorThumbnail/TrackEditorThumbnail.cpp b/Engine/Source/Editor/MovieSceneTools/Private/TrackEditorThumbnail/TrackEditorThumbnail.cpp
--- a/Engine/Source/Editor/MovieSceneTools/Private/TrackEditorThumbnail/TrackEditorThumbnail.cpp
+++ b/Engine/Source/Editor/MovieSceneTools/Private/TrackEditorThumbnail/TrackEditorThumbnail.cpp
@@ -420,8 +420,7 @@ void FTrackEditorThumbnailCache::DrawViewportThumbnail(FTrackEditorThumbnail& Tr
 	ViewFamily.Views.Add(NewView);
 
 	const float GlobalResolutionFraction = 1.f;
-	const bool  AllowPostProcessSettingsScreenPercentage = false;
-	ViewFamily.SetScreenPercentageInterface(new FLegacyScreenPercentageDriver(ViewFamily, GlobalResolutionFraction, AllowPostProcessSettingsScreenPercentage));
+	ViewFamily.SetScreenPercentageInterface(new FLegacyScreenPercentageDriver(ViewFamily, GlobalResolutionFraction));
 
 	FCanvas Canvas(TrackEditorThumbnail.GetRenderTarget(), nullptr, FApp::GetCurrentTime() - GStartTime, FApp::GetDeltaTime(), FApp::GetCurrentTime() - GStartTime, World->Scene->GetFeatureLevel());
 	Canvas.Clear(FLinearColor::Transparent);
diff --git a/Engine/Source/Editor/UnrealEd/Private/EditorViewportClient.cpp b/Engine/Source/Editor/UnrealEd/Private/EditorViewportClient.cpp
--- a/Engine/Source/Editor/UnrealEd/Private/EditorViewportClient.cpp
+++ b/Engine/Source/Editor/UnrealEd/Private/EditorViewportClient.cpp
@@ -2670,8 +2670,8 @@ int32 FEditorViewportClient::GetPreviewScreenPercentage() const
 	// We expose the resolution fraction derived from DPI, to not lie to the artist when screen percentage = 100%.
 	return FMath::RoundToInt(FMath::Clamp(
 		PreviewResolutionFraction,
-		FSceneViewScreenPercentageConfig::kMinTAAUpsampleResolutionFraction,
-		FSceneViewScreenPercentageConfig::kMaxTAAUpsampleResolutionFraction) * 100.0f);
+		ISceneViewFamilyScreenPercentage::kMinTAAUpsampleResolutionFraction,
+		ISceneViewFamilyScreenPercentage::kMaxTAAUpsampleResolutionFraction) * 100.0f);
 }
 
 void FEditorViewportClient::SetPreviewScreenPercentage(int32 PreviewScreenPercentage)
@@ -3858,7 +3858,7 @@ void FEditorViewportClient::Draw(FViewport* InViewport, FCanvas* Canvas)
 
 			// In editor viewport, we ignore r.ScreenPercentage and FPostProcessSettings::ScreenPercentage by design.
 			ViewFamily.SetScreenPercentageInterface(new FLegacyScreenPercentageDriver(
-				ViewFamily, GlobalResolutionFraction, /* AllowPostProcessSettingsScreenPercentage = */ false));
+				ViewFamily, GlobalResolutionFraction));
 		}
 
 		check(ViewFamily.GetScreenPercentageInterface() != nullptr);
diff --git a/Engine/Source/Editor/UnrealEd/Private/SCommonEditorViewportToolbarBase.cpp b/Engine/Source/Editor/UnrealEd/Private/SCommonEditorViewportToolbarBase.cpp
--- a/Engine/Source/Editor/UnrealEd/Private/SCommonEditorViewportToolbarBase.cpp
+++ b/Engine/Source/Editor/UnrealEd/Private/SCommonEditorViewportToolbarBase.cpp
@@ -388,8 +388,8 @@ void SCommonEditorViewportToolbarBase::OnFOVValueChanged(float NewValue) const
 
 TSharedRef<SWidget> SCommonEditorViewportToolbarBase::GenerateScreenPercentageMenu() const
 {
-	const int32 PreviewScreenPercentageMin = FSceneViewScreenPercentageConfig::kMinTAAUpsampleResolutionFraction * 100.0f;
-	const int32 PreviewScreenPercentageMax = FSceneViewScreenPercentageConfig::kMaxTAAUpsampleResolutionFraction * 100.0f;
+	const int32 PreviewScreenPercentageMin = ISceneViewFamilyScreenPercentage::kMinTAAUpsampleResolutionFraction * 100.0f;
+	const int32 PreviewScreenPercentageMax = ISceneViewFamilyScreenPercentage::kMaxTAAUpsampleResolutionFraction * 100.0f;
 
 	return
 		SNew(SBox)
diff --git a/Engine/Source/Editor/UnrealEd/Private/ThumbnailRendering/ThumbnailRenderer.cpp b/Engine/Source/Editor/UnrealEd/Private/ThumbnailRendering/ThumbnailRenderer.cpp
--- a/Engine/Source/Editor/UnrealEd/Private/ThumbnailRendering/ThumbnailRenderer.cpp
+++ b/Engine/Source/Editor/UnrealEd/Private/ThumbnailRendering/ThumbnailRenderer.cpp
@@ -17,7 +17,7 @@ void UThumbnailRenderer::RenderViewFamily(FCanvas* Canvas, FSceneViewFamily* Vie
 {
 	ViewFamily->EngineShowFlags.ScreenPercentage = false;
 	ViewFamily->SetScreenPercentageInterface(new FLegacyScreenPercentageDriver(
-		*ViewFamily, /* GlobalResolutionFraction = */ 1.0f, /* AllowPostProcessSettingsScreenPercentage = */ false));
+		*ViewFamily, /* GlobalResolutionFraction = */ 1.0f));
 
 	GetRendererModule().BeginRenderingViewFamily(Canvas, ViewFamily);
 }
diff --git a/Engine/Source/Runtime/Engine/Classes/Engine/Scene.h b/Engine/Source/Runtime/Engine/Classes/Engine/Scene.h
--- a/Engine/Source/Runtime/Engine/Classes/Engine/Scene.h
+++ b/Engine/Source/Runtime/Engine/Classes/Engine/Scene.h
@@ -1116,8 +1116,8 @@ struct FPostProcessSettings
 	UPROPERTY(EditAnywhere, BlueprintReadWrite, Category=Overrides, meta=(PinHiddenByDefault, InlineEditConditionToggle))
 	uint8 bOverride_MotionBlurPerObjectSize:1;
 
-	UPROPERTY(EditAnywhere, BlueprintReadWrite, Category=Overrides, meta=(PinHiddenByDefault, InlineEditConditionToggle))
-	uint8 bOverride_ScreenPercentage:1;
+	UPROPERTY()
+	uint8 bOverride_ScreenPercentage_DEPRECATED:1;
 
 	UPROPERTY(EditAnywhere, BlueprintReadWrite, Category=Overrides, meta=(PinHiddenByDefault, InlineEditConditionToggle))
 	uint8 bOverride_ReflectionMethod:1;
@@ -1980,15 +1980,8 @@ struct FPostProcessSettings
 	UPROPERTY()
 	float LPVDirectionalOcclusionFadeRange_DEPRECATED;
 
-	/**
-	* To render with lower or high resolution than it is presented,
-	* controlled by console variable,
-	* 100:off, needs to be <99 to get upsampling and lower to get performance,
-	* >100 for super sampling (slower but higher quality),
-	* only applied in game
-	*/
-	UPROPERTY(interp, BlueprintReadWrite, Category="Rendering Features|Misc", meta=(ClampMin = "0.0", ClampMax = "400.0", editcondition = "bOverride_ScreenPercentage"))
-	float ScreenPercentage;
+	UPROPERTY()
+	float ScreenPercentage_DEPRECATED;
 
 
 
diff --git a/Engine/Source/Runtime/Engine/Private/DynamicResolution.cpp b/Engine/Source/Runtime/Engine/Private/DynamicResolution.cpp
--- a/Engine/Source/Runtime/Engine/Private/DynamicResolution.cpp
+++ b/Engine/Source/Runtime/Engine/Private/DynamicResolution.cpp
@@ -471,28 +471,23 @@ public:
 		return Proxy->GetResolutionFractionUpperBound();
 	}
 
-	virtual ISceneViewFamilyScreenPercentage* Fork_GameThread(const class FSceneViewFamily& ForkedViewFamily) const override
-	{
-		check(IsInGameThread());
-
-		return new FDefaultDynamicResolutionDriver(Proxy, ForkedViewFamily);
-	}
-
-	virtual void ComputePrimaryResolutionFractions_RenderThread(TArray<FSceneViewScreenPercentageConfig>& OutViewScreenPercentageConfigs) const override
+	virtual float GetPrimaryResolutionFraction_RenderThread() const override
 	{
 		check(IsInRenderingThread());
 
 		if (!ViewFamily.EngineShowFlags.ScreenPercentage)
 		{
-			return;
+			return 1.0f;
 		}
 
-		float GlobalResolutionFraction = Proxy->QueryCurentFrameResolutionFraction_RenderThread();
+		return Proxy->QueryCurentFrameResolutionFraction_RenderThread();
+	}
 
-		for (int32 i = 0; i < OutViewScreenPercentageConfigs.Num(); i++)
-		{
-			OutViewScreenPercentageConfigs[i].PrimaryResolutionFraction = GlobalResolutionFraction;
-		}
+	virtual ISceneViewFamilyScreenPercentage* Fork_GameThread(const class FSceneViewFamily& ForkedViewFamily) const override
+	{
+		check(IsInGameThread());
+
+		return new FDefaultDynamicResolutionDriver(Proxy, ForkedViewFamily);
 	}
 
 private:
diff --git a/Engine/Source/Runtime/Engine/Private/GameViewportClient.cpp b/Engine/Source/Runtime/Engine/Private/GameViewportClient.cpp
--- a/Engine/Source/Runtime/Engine/Private/GameViewportClient.cpp
+++ b/Engine/Source/Runtime/Engine/Private/GameViewportClient.cpp
@@ -1395,7 +1395,6 @@ void UGameViewportClient::Draw(FViewport* InViewport, FCanvas* SceneCanvas)
 				ViewFamily.SetScreenPercentageInterface(new FLegacyScreenPercentageDriver(
 					ViewFamily,
 					DynamicResolutionStateInfos.ResolutionFractionApproximation,
-					/* AllowPostProcessSettingsScreenPercentage = */ false,
 					DynamicResolutionStateInfos.ResolutionFractionUpperBound));
 
 				bUsesDynamicResolution = true;
@@ -1664,7 +1663,7 @@ void UGameViewportClient::Draw(FViewport* InViewport, FCanvas* SceneCanvas)
 			}
 
 			ViewFamily.SetScreenPercentageInterface(new FLegacyScreenPercentageDriver(
-				ViewFamily, GlobalResolutionFraction, AllowPostProcessSettingsScreenPercentage));
+				ViewFamily, GlobalResolutionFraction));
 		}
 
 		check(ViewFamily.GetScreenPercentageInterface() != nullptr);
diff --git a/Engine/Source/Runtime/Engine/Private/LegacyScreenPercentageDriver.cpp b/Engine/Source/Runtime/Engine/Private/LegacyScreenPercentageDriver.cpp
--- a/Engine/Source/Runtime/Engine/Private/LegacyScreenPercentageDriver.cpp
+++ b/Engine/Source/Runtime/Engine/Private/LegacyScreenPercentageDriver.cpp
@@ -11,14 +11,12 @@
 FLegacyScreenPercentageDriver::FLegacyScreenPercentageDriver(
 	const FSceneViewFamily& InViewFamily,
 	float InGlobalResolutionFraction,
-	bool InAllowPostProcessSettingsScreenPercentage,
 	float InGlobalResolutionFractionUpperBound)
 	: ViewFamily(InViewFamily)
 	, GlobalResolutionFraction(InGlobalResolutionFraction)
 	, GlobalResolutionFractionUpperBound(InGlobalResolutionFractionUpperBound)
-	, AllowPostProcessSettingsScreenPercentage(InAllowPostProcessSettingsScreenPercentage)
 {
-	if (AllowPostProcessSettingsScreenPercentage || GlobalResolutionFraction != 1.0f)
+	if (GlobalResolutionFraction != 1.0f)
 	{
 		check(ViewFamily.EngineShowFlags.ScreenPercentage);
 	}
@@ -31,57 +29,35 @@ float FLegacyScreenPercentageDriver::GetPrimaryResolutionFractionUpperBound() co
 		return 1.0f;
 	}
 
-	float ResolutionFractionUpperBound = 0;
-
-	for (int32 i = 0; i < ViewFamily.Views.Num(); i++)
-	{
-		float ResolutionFraction = GlobalResolutionFractionUpperBound;
-		if (AllowPostProcessSettingsScreenPercentage)
-		{
-			ResolutionFraction *= ViewFamily.Views[i]->FinalPostProcessSettings.ScreenPercentage / 100.0f;
-		}
-		ResolutionFractionUpperBound = FMath::Max(ResolutionFractionUpperBound, ResolutionFraction);
-	}
-		
 	return FMath::Clamp(
-		ResolutionFractionUpperBound,
-		FSceneViewScreenPercentageConfig::kMinResolutionFraction,
-		FSceneViewScreenPercentageConfig::kMaxResolutionFraction);
-}
-
-ISceneViewFamilyScreenPercentage* FLegacyScreenPercentageDriver::Fork_GameThread(
-	const FSceneViewFamily& ForkedViewFamily) const
-{
-	check(IsInGameThread());
-
-	return new FLegacyScreenPercentageDriver(
-		ForkedViewFamily, GlobalResolutionFraction, AllowPostProcessSettingsScreenPercentage, GlobalResolutionFractionUpperBound);
+		GlobalResolutionFractionUpperBound,
+		ISceneViewFamilyScreenPercentage::kMinResolutionFraction,
+		ISceneViewFamilyScreenPercentage::kMaxResolutionFraction);
 }
 
-void FLegacyScreenPercentageDriver::ComputePrimaryResolutionFractions_RenderThread(
-	TArray<FSceneViewScreenPercentageConfig>& OutViewScreenPercentageConfigs) const
+float FLegacyScreenPercentageDriver::GetPrimaryResolutionFraction_RenderThread() const
 {
 	check(IsInRenderingThread());
 
 	// Early return if no screen percentage should be done.
 	if (!ViewFamily.EngineShowFlags.ScreenPercentage)
 	{
-		return;
+		return 1.0f;
 	}
 
-	for (int32 i = 0; i < ViewFamily.Views.Num(); i++)
-	{
-		float ResolutionFraction = GlobalResolutionFraction;
-		if (AllowPostProcessSettingsScreenPercentage)
-		{
-			ResolutionFraction *= ViewFamily.Views[i]->FinalPostProcessSettings.ScreenPercentage / 100.0f;
-		}
+	return FMath::Clamp(
+		GlobalResolutionFraction,
+		ISceneViewFamilyScreenPercentage::kMinResolutionFraction,
+		ISceneViewFamilyScreenPercentage::kMaxResolutionFraction);
+}
 
-		OutViewScreenPercentageConfigs[i].PrimaryResolutionFraction = FMath::Clamp(
-			ResolutionFraction,
-			FSceneViewScreenPercentageConfig::kMinResolutionFraction,
-			FSceneViewScreenPercentageConfig::kMaxResolutionFraction);
-	}
+ISceneViewFamilyScreenPercentage* FLegacyScreenPercentageDriver::Fork_GameThread(
+	const FSceneViewFamily& ForkedViewFamily) const
+{
+	check(IsInGameThread());
+
+	return new FLegacyScreenPercentageDriver(
+		ForkedViewFamily, GlobalResolutionFraction, GlobalResolutionFractionUpperBound);
 }
 
 // static
diff --git a/Engine/Source/Runtime/Engine/Private/Scene.cpp b/Engine/Source/Runtime/Engine/Private/Scene.cpp
--- a/Engine/Source/Runtime/Engine/Private/Scene.cpp
+++ b/Engine/Source/Runtime/Engine/Private/Scene.cpp
@@ -565,7 +565,7 @@ FPostProcessSettings::FPostProcessSettings()
 	MotionBlurMax = 5.0f;
 	MotionBlurTargetFPS = 30;
 	MotionBlurPerObjectSize = 0.5f;
-	ScreenPercentage = 100.0f;
+	ScreenPercentage_DEPRECATED = 100.0f;
 	ReflectionsType_DEPRECATED = EReflectionsType::RayTracing;
 	ReflectionMethod = EReflectionMethod::Lumen;
 	LumenReflectionQuality = 1;
@@ -751,7 +751,7 @@ FPostProcessSettings::FPostProcessSettings(const FPostProcessSettings& Settings)
 	, bOverride_MotionBlurAmount(Settings.bOverride_MotionBlurAmount)
 	, bOverride_MotionBlurMax(Settings.bOverride_MotionBlurMax)
 	, bOverride_MotionBlurPerObjectSize(Settings.bOverride_MotionBlurPerObjectSize)
-	, bOverride_ScreenPercentage(Settings.bOverride_ScreenPercentage)
+	, bOverride_ScreenPercentage_DEPRECATED(Settings.bOverride_ScreenPercentage_DEPRECATED)
 	, bOverride_ReflectionMethod(Settings.bOverride_ReflectionMethod)
 	, bOverride_LumenReflectionQuality(Settings.bOverride_LumenReflectionQuality)
 	, bOverride_ScreenSpaceReflectionIntensity(Settings.bOverride_ScreenSpaceReflectionIntensity)
@@ -957,7 +957,7 @@ FPostProcessSettings::FPostProcessSettings(const FPostProcessSettings& Settings)
 	, PathTracingMaxPathExposure(Settings.PathTracingMaxPathExposure)
 	, PathTracingEnableDenoiser(Settings.PathTracingEnableDenoiser)
 
-	, ScreenPercentage(Settings.ScreenPercentage)
+	, ScreenPercentage_DEPRECATED(Settings.ScreenPercentage_DEPRECATED)
 
 	, WeightedBlendables(Settings.WeightedBlendables)
 	//, Blendables_DEPRECATED(Settings.Blendables_DEPRECATED)
diff --git a/Engine/Source/Runtime/Engine/Private/SceneView.cpp b/Engine/Source/Runtime/Engine/Private/SceneView.cpp
--- a/Engine/Source/Runtime/Engine/Private/SceneView.cpp
+++ b/Engine/Source/Runtime/Engine/Private/SceneView.cpp
@@ -1471,7 +1471,6 @@ void FSceneView::OverridePostProcessSettings(const FPostProcessSettings& Src, fl
 		LERP_PP(MotionBlurAmount);
 		LERP_PP(MotionBlurMax);
 		LERP_PP(MotionBlurPerObjectSize);
-		LERP_PP(ScreenPercentage);
 		LERP_PP(ScreenSpaceReflectionQuality);
 		LERP_PP(ScreenSpaceReflectionIntensity);
 		LERP_PP(ScreenSpaceReflectionMaxRoughness);
@@ -2126,15 +2125,6 @@ void FSceneView::EndFinalPostprocessSettings(const FSceneViewInitOptions& ViewIn
 	}
 #endif // WITH_EDITOR
 
-	if (Family->EngineShowFlags.ScreenPercentage)
-	{
-		FinalPostProcessSettings.ScreenPercentage = FMath::Clamp(FinalPostProcessSettings.ScreenPercentage, 1.f, 400.0f);
-	}
-	else
-	{
-		FinalPostProcessSettings.ScreenPercentage = 100;
-	}
-
 	check(VerifyMembersChecks());
 
 	// If we are using FFT bloom and we don't have a convolution texture load the default one.
diff --git a/Engine/Source/Runtime/Engine/Public/LegacyScreenPercentageDriver.h b/Engine/Source/Runtime/Engine/Public/LegacyScreenPercentageDriver.h
--- a/Engine/Source/Runtime/Engine/Public/LegacyScreenPercentageDriver.h
+++ b/Engine/Source/Runtime/Engine/Public/LegacyScreenPercentageDriver.h
@@ -17,15 +17,13 @@ class ENGINE_API FLegacyScreenPercentageDriver : public ISceneViewFamilyScreenPe
 public:
 	FORCEINLINE FLegacyScreenPercentageDriver(
 		const FSceneViewFamily& InViewFamily,
-		float InGlobalResolutionFraction,
-		bool InAllowPostProcessSettingsScreenPercentage)
-		: FLegacyScreenPercentageDriver(InViewFamily, InGlobalResolutionFraction, InAllowPostProcessSettingsScreenPercentage, InGlobalResolutionFraction)
+		float InGlobalResolutionFraction)
+		: FLegacyScreenPercentageDriver(InViewFamily, InGlobalResolutionFraction, InGlobalResolutionFraction)
 	{ }
 
 	FLegacyScreenPercentageDriver(
 		const FSceneViewFamily& InViewFamily,
 		float InGlobalResolutionFraction,
-		bool InAllowPostProcessSettingsScreenPercentage,
 		float InGlobalResolutionFractionUpperBound);
 
 	/** Gets the view rect fraction from the r.ScreenPercentage cvar. */
@@ -42,13 +40,9 @@ private:
 	// ViewRect fraction to apply to all view of the view family.
 	const float GlobalResolutionFractionUpperBound;
 
-	// Whether FPostProcessSettings::ScreenPercentage should be applied or not.
-	const bool AllowPostProcessSettingsScreenPercentage;
-
 
 	// Implements ISceneViewFamilyScreenPercentage
 	virtual float GetPrimaryResolutionFractionUpperBound() const override;
+	virtual float GetPrimaryResolutionFraction_RenderThread() const override;
 	virtual ISceneViewFamilyScreenPercentage* Fork_GameThread(const class FSceneViewFamily& ForkedViewFamily) const override;
-	virtual void ComputePrimaryResolutionFractions_RenderThread(
-		TArray<FSceneViewScreenPercentageConfig>& OutViewScreenPercentageConfigs) const override;
 };
diff --git a/Engine/Source/Runtime/Engine/Public/SceneView.h b/Engine/Source/Runtime/Engine/Public/SceneView.h
--- a/Engine/Source/Runtime/Engine/Public/SceneView.h
+++ b/Engine/Source/Runtime/Engine/Public/SceneView.h
@@ -1405,18 +1405,16 @@ struct FDisplayInternalsData
 
 //////////////////////////////////////////////////////////////////////////
 
-/** Exposed screen percentage settings to ISceneViewFamilyScreenPercentage. */
-struct FSceneViewScreenPercentageConfig
-{
-	// Screen percentage / 100 to apply to a given view of the view family.
-	float PrimaryResolutionFraction;
-
-
-	FSceneViewScreenPercentageConfig()
-		: PrimaryResolutionFraction(1.f)
-	{ }
-
 
+/*
+ * Game thread and render thread interface that takes care of a FSceneViewFamily's screen percentage.
+ *
+ * The renderer reserves the right to delete and replace the view family's screen percentage interface
+ * for testing purposes with the r.Test.OverrideScreenPercentageInterface CVar.
+ */
+class ENGINE_API ISceneViewFamilyScreenPercentage
+{
+public:
 	// Sets the minimal and max screen percentage.
 	static constexpr float kMinResolutionFraction = 0.01f;
 	static constexpr float kMaxResolutionFraction = 4.0f;
@@ -1429,24 +1427,13 @@ struct FSceneViewScreenPercentageConfig
 	static constexpr float kMinTAAUpsampleResolutionFraction = 0.5f;
 	static constexpr float kMaxTAAUpsampleResolutionFraction = 2.0f;
 
+#if DO_CHECK || USING_CODE_ANALYSIS
+	static bool IsValidResolutionFraction(float ResolutionFraction)
+	{
+		return ResolutionFraction >= kMinResolutionFraction && ResolutionFraction <= kMaxResolutionFraction;
+	}
+#endif
 
-	#if DO_CHECK || USING_CODE_ANALYSIS
-		static bool IsValidResolutionFraction(float ResolutionFraction)
-		{
-			return ResolutionFraction >= kMinResolutionFraction && ResolutionFraction <= kMaxResolutionFraction;
-		}
-	#endif
-};
-
-
-/*
- * Game thread and render thread interface that takes care of a FSceneViewFamily's screen percentage.
- *
- * The renderer reserves the right to delete and replace the view family's screen percentage interface
- * for testing purposes with the r.Test.OverrideScreenPercentageInterface CVar.
- */
-class ENGINE_API ISceneViewFamilyScreenPercentage
-{
 protected:
 	/** 
 	 * Called by the destructor of the view family.
@@ -1455,20 +1442,19 @@ protected:
 	virtual ~ISceneViewFamilyScreenPercentage() {};
 
 	/** 
-	 * Method to know the maximum value that can be set in FSceneViewScreenPercentageConfig::ResolutionFraction.
+	 * Method to know the maximum value that can be returned by GetPrimaryResolutionFraction_RenderThread().
 	 * Can be called on game or rendering thread. This should return >= 1 if screen percentage show flag is disabled.
 	 */
 	virtual float GetPrimaryResolutionFractionUpperBound() const = 0;
 
-	/** Create a new screen percentage interface for a new view family. */
-	virtual ISceneViewFamilyScreenPercentage* Fork_GameThread(const class FSceneViewFamily& ViewFamily) const = 0;
-
 	/**
 	 * Setup view family's view's screen percentage on rendering thread.
 	 * This should leave ResolutionFraction == 1 if screen percentage show flag is disabled.
-	 * @param OutViewScreenPercentageConfigs Screen percentage config to set on the view of the view family.
 	 */
-	virtual void ComputePrimaryResolutionFractions_RenderThread(TArray<FSceneViewScreenPercentageConfig>& OutViewScreenPercentageConfigs) const = 0;
+	virtual float GetPrimaryResolutionFraction_RenderThread() const = 0;
+
+	/** Create a new screen percentage interface for a new view family. */
+	virtual ISceneViewFamilyScreenPercentage* Fork_GameThread(const class FSceneViewFamily& ViewFamily) const = 0;
 
 	friend class FSceneViewFamily;
 	friend class FSceneRenderer;
@@ -1727,13 +1713,13 @@ public:
 
 	FORCEINLINE bool AllowTranslucencyAfterDOF() const { return bAllowTranslucencyAfterDOF; }
 
-	/* Returns the maximum FSceneViewScreenPercentageConfig::PrimaryResolutionFraction. */
+	/* Returns the maximum PrimaryResolutionFraction. */
 	FORCEINLINE float GetPrimaryResolutionFractionUpperBound() const
 	{
 		check(ScreenPercentageInterface != nullptr);
 		float PrimaryUpperBoundFraction = ScreenPercentageInterface->GetPrimaryResolutionFractionUpperBound();
 
-		checkf(FSceneViewScreenPercentageConfig::IsValidResolutionFraction(PrimaryUpperBoundFraction),
+		checkf(ISceneViewFamilyScreenPercentage::IsValidResolutionFraction(PrimaryUpperBoundFraction),
 			TEXT("ISceneViewFamilyScreenPercentage::GetPrimaryResolutionFractionUpperBound()")
 			TEXT(" should return a valide value."));
 
diff --git a/Engine/Source/Runtime/Renderer/Private/PlanarReflectionRendering.cpp b/Engine/Source/Runtime/Renderer/Private/PlanarReflectionRendering.cpp
--- a/Engine/Source/Runtime/Renderer/Private/PlanarReflectionRendering.cpp
+++ b/Engine/Source/Runtime/Renderer/Private/PlanarReflectionRendering.cpp
@@ -633,10 +633,6 @@ void FScene::UpdatePlanarReflectionContents(UPlanarReflectionComponent* CaptureC
 
 			CaptureComponent->ProjectionWithExtraFOV[ViewIndex] = SceneCaptureViewInfo[ViewIndex].ProjectionMatrix;
 
-			// Plumb down the main view's screen percentage to the planar reflection.
-			SceneRenderer->Views[ViewIndex].FinalPostProcessSettings.ScreenPercentage =
-				MainSceneRenderer.Views[ViewIndex].FinalPostProcessSettings.ScreenPercentage;
-
 			const bool bIsStereo = IStereoRendering::IsStereoEyeView(MainSceneRenderer.Views[0]);
 
 			const FMatrix ProjectionMatrix = SceneCaptureViewInfo[ViewIndex].ProjectionMatrix;
diff --git a/Engine/Source/Runtime/Renderer/Private/PostProcess/DiaphragmDOF.cpp b/Engine/Source/Runtime/Renderer/Private/PostProcess/DiaphragmDOF.cpp
--- a/Engine/Source/Runtime/Renderer/Private/PostProcess/DiaphragmDOF.cpp
+++ b/Engine/Source/Runtime/Renderer/Private/PostProcess/DiaphragmDOF.cpp
@@ -729,9 +729,9 @@ class FDiaphragmDOFCocFlattenCS : public FDiaphragmDOFShader
 		SHADER_PARAMETER_STRUCT(FDOFTileClassificationUAVs, TileOutput)
 	END_SHADER_PARAMETER_STRUCT()
 
-	static_assert(FSceneViewScreenPercentageConfig::kMinTAAUpsampleResolutionFraction == 0.5f,
+	static_assert(ISceneViewFamilyScreenPercentage::kMinTAAUpsampleResolutionFraction == 0.5f,
 		"Gather4 shader permutation assumes with min TAAU screen percentage = 50%.");
-	static_assert(FSceneViewScreenPercentageConfig::kMaxTAAUpsampleResolutionFraction == 2.0f,
+	static_assert(ISceneViewFamilyScreenPercentage::kMaxTAAUpsampleResolutionFraction == 2.0f,
 		"Gather4 shader permutation assumes with max TAAU screen percentage = 200%.");
 };
 
diff --git a/Engine/Source/Runtime/Renderer/Private/PostProcess/TemporalAA.cpp b/Engine/Source/Runtime/Renderer/Private/PostProcess/TemporalAA.cpp
--- a/Engine/Source/Runtime/Renderer/Private/PostProcess/TemporalAA.cpp
+++ b/Engine/Source/Runtime/Renderer/Private/PostProcess/TemporalAA.cpp
@@ -921,11 +921,11 @@ public:
 
 	virtual float GetMinUpsampleResolutionFraction() const override
 	{
-		return FSceneViewScreenPercentageConfig::kMinTAAUpsampleResolutionFraction;
+		return ISceneViewFamilyScreenPercentage::kMinTAAUpsampleResolutionFraction;
 	}
 	virtual float GetMaxUpsampleResolutionFraction() const override
 	{
-		return FSceneViewScreenPercentageConfig::kMaxTAAUpsampleResolutionFraction;
+		return ISceneViewFamilyScreenPercentage::kMaxTAAUpsampleResolutionFraction;
 	}
 };
 
diff --git a/Engine/Source/Runtime/Renderer/Private/ReflectionEnvironmentCapture.cpp b/Engine/Source/Runtime/Renderer/Private/ReflectionEnvironmentCapture.cpp
--- a/Engine/Source/Runtime/Renderer/Private/ReflectionEnvironmentCapture.cpp
+++ b/Engine/Source/Runtime/Renderer/Private/ReflectionEnvironmentCapture.cpp
@@ -1441,7 +1441,7 @@ void CaptureSceneIntoScratchCubemap(
 		ViewFamily.Views.Add(View);
 
 		ViewFamily.SetScreenPercentageInterface(new FLegacyScreenPercentageDriver(
-			ViewFamily, /* GlobalResolutionFraction = */ 1.0f, /* AllowPostProcessSettingsScreenPercentage = */ false));
+			ViewFamily, /* GlobalResolutionFraction = */ 1.0f));
 
 		FSceneRenderer* SceneRenderer = FSceneRenderer::CreateSceneRenderer(&ViewFamily, NULL);
 
diff --git a/Engine/Source/Runtime/Renderer/Private/SceneCaptureRendering.cpp b/Engine/Source/Runtime/Renderer/Private/SceneCaptureRendering.cpp
--- a/Engine/Source/Runtime/Renderer/Private/SceneCaptureRendering.cpp
+++ b/Engine/Source/Runtime/Renderer/Private/SceneCaptureRendering.cpp
@@ -675,7 +675,7 @@ static FSceneRenderer* CreateSceneRendererForSceneCapture(
 	// Screen percentage is still not supported in scene capture.
 	ViewFamily.EngineShowFlags.ScreenPercentage = false;
 	ViewFamily.SetScreenPercentageInterface(new FLegacyScreenPercentageDriver(
-		ViewFamily, /* GlobalResolutionFraction = */ 1.0f, /* AllowPostProcessSettingsScreenPercentage = */ false));
+		ViewFamily, /* GlobalResolutionFraction = */ 1.0f));
 
 	return FSceneRenderer::CreateSceneRenderer(&ViewFamily, nullptr);
 }
diff --git a/Engine/Source/Runtime/Renderer/Private/SceneRendering.cpp b/Engine/Source/Runtime/Renderer/Private/SceneRendering.cpp
--- a/Engine/Source/Runtime/Renderer/Private/SceneRendering.cpp
+++ b/Engine/Source/Runtime/Renderer/Private/SceneRendering.cpp
@@ -520,17 +520,17 @@ public:
 		}
 
 		return new FLegacyScreenPercentageDriver(
-			ForkedViewFamily, /* GlobalResolutionFraction = */ MaxResolutionFraction, /* AllowPostProcessSettingsScreenPercentage = */ false);
+			ForkedViewFamily, /* GlobalResolutionFraction = */ MaxResolutionFraction);
 	}
 
-	virtual void ComputePrimaryResolutionFractions_RenderThread(TArray<FSceneViewScreenPercentageConfig>& OutViewScreenPercentageConfigs) const override
+	virtual float GetPrimaryResolutionFraction_RenderThread() const override
 	{
 		check(IsInRenderingThread());
 
 		// Early return if no screen percentage should be done.
 		if (!ViewFamily.EngineShowFlags.ScreenPercentage)
 		{
-			return;
+			return 1.0f;
 		}
 
 		uint32 FrameId = 0;
@@ -540,12 +540,7 @@ public:
 		{
 			FrameId = ViewState->GetFrameIndex(8);
 		}
-		float ResolutionFraction = FrameId == 0 ? MaxResolutionFraction : FMath::Lerp(MinResolutionFraction, MaxResolutionFraction, 0.5f + 0.5f * FMath::Cos((FrameId + 0.25) * PI / 8));
-	
-		for (int32 i = 0; i < ViewFamily.Views.Num(); i++)
-		{
-			OutViewScreenPercentageConfigs[i].PrimaryResolutionFraction = ResolutionFraction;
-		}
+		return FrameId == 0 ? MaxResolutionFraction : FMath::Lerp(MinResolutionFraction, MaxResolutionFraction, 0.5f + 0.5f * FMath::Cos((FrameId + 0.25) * PI / 8));
 	}
 
 private:
@@ -2446,7 +2441,7 @@ FIntPoint FSceneRenderer::ApplyResolutionFraction(const FSceneViewFamily& ViewFa
 {
 	FIntPoint ViewSize;
 
-	// CeilToInt so tha view size is at least 1x1 if ResolutionFraction == FSceneViewScreenPercentageConfig::kMinResolutionFraction.
+	// CeilToInt so tha view size is at least 1x1 if ResolutionFraction == ISceneViewFamilyScreenPercentage::kMinResolutionFraction.
 	ViewSize.X = FMath::CeilToInt(UnscaledViewSize.X * ResolutionFraction);
 	ViewSize.Y = FMath::CeilToInt(UnscaledViewSize.Y * ResolutionFraction);
 
@@ -2547,9 +2542,6 @@ void FSceneRenderer::PrepareViewRectsForRendering()
 		return;
 	}
 
-	TArray<FSceneViewScreenPercentageConfig> ViewScreenPercentageConfigs;
-	ViewScreenPercentageConfigs.Reserve(Views.Num());
-
 	// Checks that view rects were still not initialized.
 	for (FViewInfo& View : Views)
 	{
@@ -2571,22 +2563,11 @@ void FSceneRenderer::PrepareViewRectsForRendering()
 				View.AntiAliasingMethod = AAM_None;
 			}
 		}
-
-		FSceneViewScreenPercentageConfig Config;
-		ViewScreenPercentageConfigs.Add(FSceneViewScreenPercentageConfig());
 	}
 
 	check(ViewFamily.ScreenPercentageInterface);
-	ViewFamily.ScreenPercentageInterface->ComputePrimaryResolutionFractions_RenderThread(ViewScreenPercentageConfigs);
-
-	check(ViewScreenPercentageConfigs.Num() == Views.Num());
-
-	// Checks that view rects are correctly initialized.
-	for (int32 i = 0; i < Views.Num(); i++)
+	float PrimaryResolutionFraction = ViewFamily.ScreenPercentageInterface->GetPrimaryResolutionFraction_RenderThread();
 	{
-		FViewInfo& View = Views[i];
-		float PrimaryResolutionFraction = ViewScreenPercentageConfigs[i].PrimaryResolutionFraction;
-
 		// Ensure screen percentage show flag is respected. Prefer to check() rather rendering at a differen screen percentage
 		// to make sure the renderer does not lie how a frame as been rendering to a dynamic resolution heuristic.
 		if (!ViewFamily.EngineShowFlags.ScreenPercentage)
@@ -2597,11 +2578,17 @@ void FSceneRenderer::PrepareViewRectsForRendering()
 		// Make sure the screen percentage interface has not lied to the renderer about the upper bound.
 		checkf(PrimaryResolutionFraction <= ViewFamily.GetPrimaryResolutionFractionUpperBound(),
 			TEXT("ISceneViewFamilyScreenPercentage::GetPrimaryResolutionFractionUpperBound() should not lie to the renderer."));
-		
-		check(FSceneViewScreenPercentageConfig::IsValidResolutionFraction(PrimaryResolutionFraction));
-		
-		// Compute final resolution fraction.
-		float ResolutionFraction = PrimaryResolutionFraction * ViewFamily.SecondaryViewFraction;
+
+		check(ISceneViewFamilyScreenPercentage::IsValidResolutionFraction(PrimaryResolutionFraction));
+	}
+
+	// Compute final resolution fraction.
+	float ResolutionFraction = PrimaryResolutionFraction * ViewFamily.SecondaryViewFraction;
+
+	// Checks that view rects are correctly initialized.
+	for (int32 i = 0; i < Views.Num(); i++)
+	{
+		FViewInfo& View = Views[i];
 
 		FIntPoint ViewSize = ApplyResolutionFraction(ViewFamily, View.UnscaledViewRect.Size(), ResolutionFraction);
 		FIntPoint ViewRectMin = QuantizeViewRectMin(FIntPoint(
diff --git a/Engine/Source/Runtime/UMG/Private/Components/Viewport.cpp b/Engine/Source/Runtime/UMG/Private/Components/Viewport.cpp
--- a/Engine/Source/Runtime/UMG/Private/Components/Viewport.cpp
+++ b/Engine/Source/Runtime/UMG/Private/Components/Viewport.cpp
@@ -175,7 +175,7 @@ void FUMGViewportClient::Draw(FViewport* InViewport, FCanvas* Canvas)
 	FSceneView* View = CalcSceneView(&ViewFamily);
 
 	ViewFamily.SetScreenPercentageInterface(new FLegacyScreenPercentageDriver(
-		ViewFamily, GlobalResolutionFraction, /* AllowPostProcessSettingsScreenPercentage = */ false));
+		ViewFamily, GlobalResolutionFraction));
 
 	//SetupViewForRendering(ViewFamily, *View);
 
