Fixes clang compilation in OculusToolWidget.cpp

diff --git a/Engine/Plugins/Runtime/Oculus/OculusVR/Source/OculusEditor/Private/OculusToolWidget.cpp b/Engine/Plugins/Runtime/Oculus/OculusVR/Source/OculusEditor/Private/OculusToolWidget.cpp
--- a/Engine/Plugins/Runtime/Oculus/OculusVR/Source/OculusEditor/Private/OculusToolWidget.cpp
+++ b/Engine/Plugins/Runtime/Oculus/OculusVR/Source/OculusEditor/Private/OculusToolWidget.cpp
@@ -715,7 +715,7 @@ FReply SOculusToolWidget::MobileMSAAEnable(bool text)
 EVisibility SOculusToolWidget::MobileMSAAVisibility(FName tag) const
 {
 	URendererSettings* Settings = GetMutableDefault<URendererSettings>();
-	const bool bMobileMSAAValid = Settings->MobileAntiAliasing == EAntiAliasingMethod::AAM_MSAA && Settings->MSAASampleCount == ECompositingSampleCount::Four;
+	const bool bMobileMSAAValid = Settings->MobileAntiAliasing == EMobileAntiAliasingMethod::MSAA && Settings->MSAASampleCount == ECompositingSampleCount::Four;
 
 	return bMobileMSAAValid ? EVisibility::Collapsed : EVisibility::Visible;
 }
