XR: Fixed changed rendering thread assumptions from recent RDG pass changes
Rename occlusion mesh functions accordingly

diff --git a/Engine/Plugins/Runtime/Oculus/OculusVR/Source/OculusHMD/Private/OculusHMD.cpp b/Engine/Plugins/Runtime/Oculus/OculusVR/Source/OculusHMD/Private/OculusHMD.cpp
--- a/Engine/Plugins/Runtime/Oculus/OculusVR/Source/OculusHMD/Private/OculusHMD.cpp
+++ b/Engine/Plugins/Runtime/Oculus/OculusVR/Source/OculusHMD/Private/OculusHMD.cpp
@@ -1181,9 +1181,8 @@ namespace OculusHMD
 	}
 
 
-	static void DrawOcclusionMesh_RenderThread(FRHICommandList& RHICmdList, EStereoscopicPass StereoPass, const FHMDViewMesh MeshAssets[])
+	static void DrawOcclusionMesh(FRHICommandList& RHICmdList, EStereoscopicPass StereoPass, const FHMDViewMesh MeshAssets[])
 	{
-		CheckInRenderThread();
 		check(StereoPass != eSSP_FULL);
 
 		const uint32 MeshIndex = (StereoPass == eSSP_LEFT_EYE) ? 0 : 1;
@@ -1195,20 +1194,16 @@ namespace OculusHMD
 	}
 
 
-	void FOculusHMD::DrawHiddenAreaMesh_RenderThread(FRHICommandList& RHICmdList, EStereoscopicPass StereoPass) const
+	void FOculusHMD::DrawHiddenAreaMesh(FRHICommandList& RHICmdList, EStereoscopicPass StereoPass) const
 	{
-		CheckInRenderThread();
-
-		DrawOcclusionMesh_RenderThread(RHICmdList, StereoPass, HiddenAreaMeshes);
+		DrawOcclusionMesh(RHICmdList, StereoPass, HiddenAreaMeshes);
 	}
 
 
-	void FOculusHMD::DrawVisibleAreaMesh_RenderThread(FRHICommandList& RHICmdList, EStereoscopicPass StereoPass) const
-		{
-		CheckInRenderThread();
-
-		DrawOcclusionMesh_RenderThread(RHICmdList, StereoPass, VisibleAreaMeshes);
-		}
+	void FOculusHMD::DrawVisibleAreaMesh(FRHICommandList& RHICmdList, EStereoscopicPass StereoPass) const
+	{
+		DrawOcclusionMesh(RHICmdList, StereoPass, VisibleAreaMeshes);
+	}
 
 	float FOculusHMD::GetPixelDenity() const
 	{
diff --git a/Engine/Plugins/Runtime/Oculus/OculusVR/Source/OculusHMD/Private/OculusHMD.h b/Engine/Plugins/Runtime/Oculus/OculusVR/Source/OculusHMD/Private/OculusHMD.h
--- a/Engine/Plugins/Runtime/Oculus/OculusVR/Source/OculusHMD/Private/OculusHMD.h
+++ b/Engine/Plugins/Runtime/Oculus/OculusVR/Source/OculusHMD/Private/OculusHMD.h
@@ -134,8 +134,8 @@ public:
 	//virtual bool GetChromaAbCorrectionValues(FVector4& K) const override;
 	virtual bool HasHiddenAreaMesh() const override;
 	virtual bool HasVisibleAreaMesh() const override;
-	virtual void DrawHiddenAreaMesh_RenderThread(class FRHICommandList& RHICmdList, EStereoscopicPass StereoPass) const override;
-	virtual void DrawVisibleAreaMesh_RenderThread(class FRHICommandList& RHICmdList, EStereoscopicPass StereoPass) const override;
+	virtual void DrawHiddenAreaMesh(class FRHICommandList& RHICmdList, EStereoscopicPass StereoPass) const override;
+	virtual void DrawVisibleAreaMesh(class FRHICommandList& RHICmdList, EStereoscopicPass StereoPass) const override;
 	//virtual void DrawDistortionMesh_RenderThread(struct FHeadMountedDisplayPassContext& Context, const FIntPoint& TextureSize) override;
 	//virtual void UpdateScreenSettings(const FViewport* InViewport) override;
 	//virtual void UpdatePostProcessSettings(FPostProcessSettings*) override;
diff --git a/Engine/Plugins/Runtime/OpenXR/Source/OpenXRHMD/Private/OpenXRHMD.cpp b/Engine/Plugins/Runtime/OpenXR/Source/OpenXRHMD/Private/OpenXRHMD.cpp
--- a/Engine/Plugins/Runtime/OpenXR/Source/OpenXRHMD/Private/OpenXRHMD.cpp
+++ b/Engine/Plugins/Runtime/OpenXR/Source/OpenXRHMD/Private/OpenXRHMD.cpp
@@ -3063,7 +3063,7 @@ bool FOpenXRHMD::HasVisibleAreaMesh() const
 	return VisibleAreaMeshes.Num() > 0;
 }
 
-void FOpenXRHMD::DrawHiddenAreaMesh_RenderThread(class FRHICommandList& RHICmdList, EStereoscopicPass StereoPass) const
+void FOpenXRHMD::DrawHiddenAreaMesh(class FRHICommandList& RHICmdList, EStereoscopicPass StereoPass) const
 {
 	check(StereoPass != eSSP_FULL);
 
@@ -3080,7 +3080,7 @@ void FOpenXRHMD::DrawHiddenAreaMesh_RenderThread(class FRHICommandList& RHICmdLi
 	}
 }
 
-void FOpenXRHMD::DrawVisibleAreaMesh_RenderThread(class FRHICommandList& RHICmdList, EStereoscopicPass StereoPass) const
+void FOpenXRHMD::DrawVisibleAreaMesh(class FRHICommandList& RHICmdList, EStereoscopicPass StereoPass) const
 {
 	check(StereoPass != eSSP_FULL);
 
diff --git a/Engine/Plugins/Runtime/OpenXR/Source/OpenXRHMD/Private/OpenXRHMD.h b/Engine/Plugins/Runtime/OpenXR/Source/OpenXRHMD/Private/OpenXRHMD.h
--- a/Engine/Plugins/Runtime/OpenXR/Source/OpenXRHMD/Private/OpenXRHMD.h
+++ b/Engine/Plugins/Runtime/OpenXR/Source/OpenXRHMD/Private/OpenXRHMD.h
@@ -215,8 +215,8 @@ public:
 	virtual void CopyTexture_RenderThread(FRHICommandListImmediate& RHICmdList, FRHITexture2D* SrcTexture, FIntRect SrcRect, FRHITexture2D* DstTexture, FIntRect DstRect, bool bClearBlack, bool bNoAlpha) const override;
 	virtual bool HasHiddenAreaMesh() const override final;
 	virtual bool HasVisibleAreaMesh() const override final;
-	virtual void DrawHiddenAreaMesh_RenderThread(class FRHICommandList& RHICmdList, EStereoscopicPass StereoPass) const override final;
-	virtual void DrawVisibleAreaMesh_RenderThread(class FRHICommandList& RHICmdList, EStereoscopicPass StereoPass) const override final;
+	virtual void DrawHiddenAreaMesh(class FRHICommandList& RHICmdList, EStereoscopicPass StereoPass) const override final;
+	virtual void DrawVisibleAreaMesh(class FRHICommandList& RHICmdList, EStereoscopicPass StereoPass) const override final;
 	virtual void OnBeginRendering_RenderThread(FRHICommandListImmediate& RHICmdList, FSceneViewFamily& ViewFamily) override;
 	virtual void OnBeginRendering_GameThread() override;
 	virtual void OnLateUpdateApplied_RenderThread(FRHICommandListImmediate& RHICmdList, const FTransform& NewRelativeTransform) override;
diff --git a/Engine/Plugins/Runtime/Steam/SteamVR/Source/SteamVR/Private/SteamVRHMD.h b/Engine/Plugins/Runtime/Steam/SteamVR/Source/SteamVR/Private/SteamVRHMD.h
--- a/Engine/Plugins/Runtime/Steam/SteamVR/Source/SteamVR/Private/SteamVRHMD.h
+++ b/Engine/Plugins/Runtime/Steam/SteamVR/Source/SteamVR/Private/SteamVRHMD.h
@@ -142,10 +142,10 @@ public:
 	virtual bool IsChromaAbCorrectionEnabled() const override;
 
 	virtual bool HasHiddenAreaMesh() const override { return HiddenAreaMeshes[0].IsValid() && HiddenAreaMeshes[1].IsValid(); }
-	virtual void DrawHiddenAreaMesh_RenderThread(FRHICommandList& RHICmdList, EStereoscopicPass StereoPass) const override;
+	virtual void DrawHiddenAreaMesh(FRHICommandList& RHICmdList, EStereoscopicPass StereoPass) const override;
 
 	virtual bool HasVisibleAreaMesh() const override { return VisibleAreaMeshes[0].IsValid() && VisibleAreaMeshes[1].IsValid(); }
-	virtual void DrawVisibleAreaMesh_RenderThread(FRHICommandList& RHICmdList, EStereoscopicPass StereoPass) const override;
+	virtual void DrawVisibleAreaMesh(FRHICommandList& RHICmdList, EStereoscopicPass StereoPass) const override;
 
 	virtual void DrawDistortionMesh_RenderThread(struct FHeadMountedDisplayPassContext& Context, const FIntPoint& TextureSize) override;
 
diff --git a/Engine/Plugins/Runtime/Steam/SteamVR/Source/SteamVR/Private/SteamVRRender.cpp b/Engine/Plugins/Runtime/Steam/SteamVR/Source/SteamVR/Private/SteamVRRender.cpp
--- a/Engine/Plugins/Runtime/Steam/SteamVR/Source/SteamVR/Private/SteamVRRender.cpp
+++ b/Engine/Plugins/Runtime/Steam/SteamVR/Source/SteamVR/Private/SteamVRRender.cpp
@@ -71,7 +71,6 @@ bool FSteamVRHMD::IsActiveThisFrame_Internal(const FSceneViewExtensionContext& C
 
 static void DrawOcclusionMesh(FRHICommandList& RHICmdList, EStereoscopicPass StereoPass, const FHMDViewMesh MeshAssets[])
 {
-	check(IsInRenderingThread());
 	check(GEngine->StereoRenderingDevice->DeviceIsStereoEyePass(StereoPass));
 
 	const uint32 MeshIndex = GEngine->StereoRenderingDevice->GetViewIndexForPass(StereoPass);
@@ -82,12 +81,12 @@ static void DrawOcclusionMesh(FRHICommandList& RHICmdList, EStereoscopicPass Ste
 	RHICmdList.DrawIndexedPrimitive(Mesh.IndexBufferRHI, 0, 0, Mesh.NumVertices, 0, Mesh.NumTriangles, 1);
 }
 
-void FSteamVRHMD::DrawHiddenAreaMesh_RenderThread(FRHICommandList& RHICmdList, EStereoscopicPass StereoPass) const
+void FSteamVRHMD::DrawHiddenAreaMesh(FRHICommandList& RHICmdList, EStereoscopicPass StereoPass) const
 {
 	DrawOcclusionMesh(RHICmdList, StereoPass, HiddenAreaMeshes);
 }
 
-void FSteamVRHMD::DrawVisibleAreaMesh_RenderThread(FRHICommandList& RHICmdList, EStereoscopicPass StereoPass) const
+void FSteamVRHMD::DrawVisibleAreaMesh(FRHICommandList& RHICmdList, EStereoscopicPass StereoPass) const
 {
 	DrawOcclusionMesh(RHICmdList, StereoPass, VisibleAreaMeshes);
 }
diff --git a/Engine/Plugins/Runtime/WindowsMixedReality/Source/WindowsMixedRealityHMD/Private/WindowsMixedRealityHMD.cpp b/Engine/Plugins/Runtime/WindowsMixedReality/Source/WindowsMixedRealityHMD/Private/WindowsMixedRealityHMD.cpp
--- a/Engine/Plugins/Runtime/WindowsMixedReality/Source/WindowsMixedRealityHMD/Private/WindowsMixedRealityHMD.cpp
+++ b/Engine/Plugins/Runtime/WindowsMixedReality/Source/WindowsMixedRealityHMD/Private/WindowsMixedRealityHMD.cpp
@@ -1730,7 +1730,7 @@ namespace WindowsMixedReality
 		return HiddenAreaMesh[0].IsValid() && HiddenAreaMesh[1].IsValid();
 	}
 
-	void FWindowsMixedRealityHMD::DrawHiddenAreaMesh_RenderThread(FRHICommandList& RHICmdList, EStereoscopicPass StereoPass) const
+	void FWindowsMixedRealityHMD::DrawHiddenAreaMesh(FRHICommandList& RHICmdList, EStereoscopicPass StereoPass) const
 	{
 		if (StereoPass == eSSP_FULL || StereoPass == eSSP_THIRD_CAMERA_EYE)
 		{
@@ -1752,7 +1752,7 @@ namespace WindowsMixedReality
 		//return VisibleAreaMesh[0].IsValid() && VisibleAreaMesh[1].IsValid();
 	}
 
-	void FWindowsMixedRealityHMD::DrawVisibleAreaMesh_RenderThread(FRHICommandList& RHICmdList, EStereoscopicPass StereoPass) const
+	void FWindowsMixedRealityHMD::DrawVisibleAreaMesh(FRHICommandList& RHICmdList, EStereoscopicPass StereoPass) const
 	{
 		if (StereoPass == eSSP_FULL || StereoPass == eSSP_THIRD_CAMERA_EYE)
 		{
diff --git a/Engine/Plugins/Runtime/WindowsMixedReality/Source/WindowsMixedRealityHMD/Private/WindowsMixedRealityHMD.h b/Engine/Plugins/Runtime/WindowsMixedReality/Source/WindowsMixedRealityHMD/Private/WindowsMixedRealityHMD.h
--- a/Engine/Plugins/Runtime/WindowsMixedReality/Source/WindowsMixedRealityHMD/Private/WindowsMixedRealityHMD.h
+++ b/Engine/Plugins/Runtime/WindowsMixedReality/Source/WindowsMixedRealityHMD/Private/WindowsMixedRealityHMD.h
@@ -179,10 +179,10 @@ namespace WindowsMixedReality
 		}
 
 		virtual bool HasHiddenAreaMesh() const override;
-		virtual void DrawHiddenAreaMesh_RenderThread(FRHICommandList& RHICmdList, EStereoscopicPass StereoPass) const override;
+		virtual void DrawHiddenAreaMesh(FRHICommandList& RHICmdList, EStereoscopicPass StereoPass) const override;
 
 		virtual bool HasVisibleAreaMesh() const override;
-		virtual void DrawVisibleAreaMesh_RenderThread(FRHICommandList& RHICmdList, EStereoscopicPass StereoPass) const override;
+		virtual void DrawVisibleAreaMesh(FRHICommandList& RHICmdList, EStereoscopicPass StereoPass) const override;
 
 		// Spectator screen Hooks.
 		virtual FIntRect GetFullFlatEyeRect_RenderThread(FTexture2DRHIRef EyeTexture) const override;
diff --git a/Engine/Source/Runtime/HeadMountedDisplay/Public/IHeadMountedDisplay.h b/Engine/Source/Runtime/HeadMountedDisplay/Public/IHeadMountedDisplay.h
--- a/Engine/Source/Runtime/HeadMountedDisplay/Public/IHeadMountedDisplay.h
+++ b/Engine/Source/Runtime/HeadMountedDisplay/Public/IHeadMountedDisplay.h
@@ -213,13 +213,13 @@ public:
 	* Optional method to draw a view's hidden area mesh where supported.
 	* This can be used to avoid rendering pixels which are not included as input into the final distortion pass.
 	*/
-	virtual void DrawHiddenAreaMesh_RenderThread(class FRHICommandList& RHICmdList, EStereoscopicPass StereoPass) const {};
+	virtual void DrawHiddenAreaMesh(class FRHICommandList& RHICmdList, EStereoscopicPass StereoPass) const {};
 
 	/**
 	* Optional method to draw a view's visible area mesh where supported.
 	* This can be used instead of a full screen quad to avoid rendering pixels which are not included as input into the final distortion pass.
 	*/
-	virtual void DrawVisibleAreaMesh_RenderThread(class FRHICommandList& RHICmdList, EStereoscopicPass StereoPass) const {};
+	virtual void DrawVisibleAreaMesh(class FRHICommandList& RHICmdList, EStereoscopicPass StereoPass) const {};
 
 	virtual void DrawDistortionMesh_RenderThread(struct FHeadMountedDisplayPassContext& Context, const FIntPoint& TextureSize) {}
 
diff --git a/Engine/Source/Runtime/Renderer/Private/DepthRendering.cpp b/Engine/Source/Runtime/Renderer/Private/DepthRendering.cpp
--- a/Engine/Source/Runtime/Renderer/Private/DepthRendering.cpp
+++ b/Engine/Source/Runtime/Renderer/Private/DepthRendering.cpp
@@ -684,7 +684,7 @@ void FDeferredShadingSceneRenderer::RenderPrePassHMD(FRDGBuilder& GraphBuilder,
 
 				SetGraphicsPipelineState(RHICmdList, GraphicsPSOInit);
 				VertexShader->SetDepthParameter(RHICmdList, 1.0f);
-				HMDDevice->DrawHiddenAreaMesh_RenderThread(RHICmdList, View.StereoPass);
+				HMDDevice->DrawHiddenAreaMesh(RHICmdList, View.StereoPass);
 			});
 		}
 	}
diff --git a/Engine/Source/Runtime/Renderer/Private/PostProcess/SceneFilterRendering.cpp b/Engine/Source/Runtime/Renderer/Private/PostProcess/SceneFilterRendering.cpp
--- a/Engine/Source/Runtime/Renderer/Private/PostProcess/SceneFilterRendering.cpp
+++ b/Engine/Source/Runtime/Renderer/Private/PostProcess/SceneFilterRendering.cpp
@@ -262,7 +262,7 @@ void DrawHmdMesh(
 
 	if (GEngine->XRSystem->GetHMDDevice())
 	{
-		GEngine->XRSystem->GetHMDDevice()->DrawVisibleAreaMesh_RenderThread(RHICmdList, StereoView);
+		GEngine->XRSystem->GetHMDDevice()->DrawVisibleAreaMesh(RHICmdList, StereoView);
 	}
 }
 
