Fix UE4 reference in Oculus plugin source ("UE4" -> "UE" in comment).

diff --git a/Engine/Plugins/Runtime/Oculus/OculusVR/Source/OculusHMD/Public/OculusFunctionLibrary.h b/Engine/Plugins/Runtime/Oculus/OculusVR/Source/OculusHMD/Public/OculusFunctionLibrary.h
--- a/Engine/Plugins/Runtime/Oculus/OculusVR/Source/OculusHMD/Public/OculusFunctionLibrary.h
+++ b/Engine/Plugins/Runtime/Oculus/OculusVR/Source/OculusHMD/Public/OculusFunctionLibrary.h
@@ -458,7 +458,7 @@ class OCULUSHMD_API UOculusFunctionLibrary : public UBlueprintFunctionLibrary
 	static FTransform GetPlayAreaTransform();
 
 	/**
-	* Get the intersection result between a UE4 coordinate and a guardian boundary
+	* Get the intersection result between a UE coordinate and a guardian boundary
 	* @param Point					(in) Point in UE space to test against guardian boundaries
 	* @param BoundaryType			(in) An enum representing the boundary type requested, either Outer Boundary (exact guardian bounds) or PlayArea (rectangle inside the Outer Boundary)
 	*/
