Fix Android-only non-unity build error in Oculus plugin due to audio header namespace change.

diff --git a/Engine/Plugins/Runtime/Oculus/OculusVR/Source/OculusMR/Private/OculusMR_CastingCameraActor.h b/Engine/Plugins/Runtime/Oculus/OculusVR/Source/OculusMR/Private/OculusMR_CastingCameraActor.h
--- a/Engine/Plugins/Runtime/Oculus/OculusVR/Source/OculusMR/Private/OculusMR_CastingCameraActor.h
+++ b/Engine/Plugins/Runtime/Oculus/OculusVR/Source/OculusMR/Private/OculusMR_CastingCameraActor.h
@@ -135,7 +135,7 @@ private:
 	TObjectPtr<UOculusMR_State> MRState;
 
 #if PLATFORM_ANDROID
-	TArray<Audio::FAlignedFloatBuffer> AudioBuffers;
+	TArray<Audio::VectorOps::FAlignedFloatBuffer> AudioBuffers;
 	TArray<double> AudioTimes;
 
 	int SyncId;
