Allow OculusHMDModule to setup OculusPluginWrapper when Oculus Service is not installed or broken; Remove redundant OVRPlugin.dll verification from OculusEditorModule
(Submitted by Oculus)

diff --git a/Engine/Plugins/Runtime/Oculus/OculusVR/Source/OculusEditor/Private/OculusEditorModule.cpp b/Engine/Plugins/Runtime/Oculus/OculusVR/Source/OculusEditor/Private/OculusEditorModule.cpp
--- a/Engine/Plugins/Runtime/Oculus/OculusVR/Source/OculusEditor/Private/OculusEditorModule.cpp
+++ b/Engine/Plugins/Runtime/Oculus/OculusVR/Source/OculusEditor/Private/OculusEditorModule.cpp
@@ -34,15 +34,6 @@ void FOculusEditorModule::PostLoadCallback()
 
 void FOculusEditorModule::StartupModule()
 {
-	void* ModuleCheck = FPlatformProcess::GetDllHandle(TEXT("OVRPlugin.dll"));
-	if (!ModuleCheck)
-	{
-		return;
-	}
-
-	FPlatformProcess::FreeDllHandle(ModuleCheck);
-	ModuleCheck = nullptr;
-
 	bModuleValid = true;
 	RegisterSettings();
 	FOculusAssetDirectory::LoadForCook();
diff --git a/Engine/Plugins/Runtime/Oculus/OculusVR/Source/OculusHMD/Private/OculusHMDModule.cpp b/Engine/Plugins/Runtime/Oculus/OculusVR/Source/OculusHMD/Private/OculusHMDModule.cpp
--- a/Engine/Plugins/Runtime/Oculus/OculusVR/Source/OculusHMD/Private/OculusHMDModule.cpp
+++ b/Engine/Plugins/Runtime/Oculus/OculusVR/Source/OculusHMD/Private/OculusHMDModule.cpp
@@ -87,8 +87,8 @@ bool FOculusHMDModule::PreInit()
 		}
 #endif
 
-		// Init module if app can render and Oculus service is running
-		if (FApp::CanEverRender() && OculusHMD::IsOculusServiceRunning())
+		// Init module if app can render
+		if (FApp::CanEverRender())
 		{
 			// Load OVRPlugin
 			OVRPluginHandle = GetOVRPluginHandle();
