Fix Oculus plugin build break on Android (Audio::AlignedFloatBuffer renamed to Audio::FAlignedFloatBuffer).

diff --git a/Engine/Plugins/Runtime/Oculus/OculusVR/Source/OculusMR/Private/OculusMR_CastingCameraActor.h b/Engine/Plugins/Runtime/Oculus/OculusVR/Source/OculusMR/Private/OculusMR_CastingCameraActor.h
--- a/Engine/Plugins/Runtime/Oculus/OculusVR/Source/OculusMR/Private/OculusMR_CastingCameraActor.h
+++ b/Engine/Plugins/Runtime/Oculus/OculusVR/Source/OculusMR/Private/OculusMR_CastingCameraActor.h
@@ -135,7 +135,7 @@ private:
 	TObjectPtr<UOculusMR_State> MRState;
 
 #if PLATFORM_ANDROID
-	TArray<Audio::AlignedFloatBuffer> AudioBuffers;
+	TArray<Audio::FAlignedFloatBuffer> AudioBuffers;
 	TArray<double> AudioTimes;
 
 	int SyncId;
