Deletes dead unconverted shader class in OculusVR

diff --git a/Engine/Plugins/Runtime/Oculus/OculusVR/Shaders/Private/OculusStressShaders.usf b/Engine/Plugins/Runtime/Oculus/OculusVR/Shaders/Private/OculusStressShaders.usf
deleted file mode 100644
diff --git a/Engine/Plugins/Runtime/Oculus/OculusVR/Source/OculusHMD/Private/OculusHMD.cpp b/Engine/Plugins/Runtime/Oculus/OculusVR/Source/OculusHMD/Private/OculusHMD.cpp
--- a/Engine/Plugins/Runtime/Oculus/OculusVR/Source/OculusHMD/Private/OculusHMD.cpp
+++ b/Engine/Plugins/Runtime/Oculus/OculusVR/Source/OculusHMD/Private/OculusHMD.cpp
@@ -1384,10 +1384,6 @@ namespace OculusHMD
 		{
 			SpectatorScreenController->RenderSpectatorScreen_RenderThread(RHICmdList, BackBuffer, SrcTexture, WindowSize);
 		}
-
-#if OCULUS_STRESS_TESTS_ENABLED
-		FStressTester::TickGPU_RenderThread(RHICmdList, BackBuffer, SrcTexture);
-#endif
 	}
 
 	FVector2D FOculusHMD::GetEyeCenterPoint_RenderThread(int32 ViewIndex) const
diff --git a/Engine/Plugins/Runtime/Oculus/OculusVR/Source/OculusHMD/Private/OculusHMDModule.cpp b/Engine/Plugins/Runtime/Oculus/OculusVR/Source/OculusHMD/Private/OculusHMDModule.cpp
--- a/Engine/Plugins/Runtime/Oculus/OculusVR/Source/OculusHMD/Private/OculusHMDModule.cpp
+++ b/Engine/Plugins/Runtime/Oculus/OculusVR/Source/OculusHMD/Private/OculusHMDModule.cpp
@@ -34,8 +34,6 @@ FOculusHMDModule::FOculusHMDModule()
 void FOculusHMDModule::StartupModule()
 {
 	IHeadMountedDisplayModule::StartupModule();
-	FString PluginShaderDir = FPaths::Combine(IPluginManager::Get().FindPlugin(TEXT("OculusVR"))->GetBaseDir(), TEXT("Shaders"));
-	AddShaderSourceDirectoryMapping(TEXT("/Plugin/OculusVR"), PluginShaderDir);
 }
 
 void FOculusHMDModule::ShutdownModule()
diff --git a/Engine/Plugins/Runtime/Oculus/OculusVR/Source/OculusHMD/Private/OculusHMD_StressTester.cpp b/Engine/Plugins/Runtime/Oculus/OculusVR/Source/OculusHMD/Private/OculusHMD_StressTester.cpp
--- a/Engine/Plugins/Runtime/Oculus/OculusVR/Source/OculusHMD/Private/OculusHMD_StressTester.cpp
+++ b/Engine/Plugins/Runtime/Oculus/OculusVR/Source/OculusHMD/Private/OculusHMD_StressTester.cpp
@@ -96,88 +96,6 @@ public:
 static TGlobalResource<FTextureVertexDeclaration> GOculusTextureVertexDeclaration;
 
 
-//-------------------------------------------------------------------------------------------------
-// FOculusStressShadersPS
-//-------------------------------------------------------------------------------------------------
-
-// DO NOT MERGE -- Removing Packaging Issue
-#if 0
-
-class FOculusStressShadersPS : public FGlobalShader
-{
-
-	DECLARE_SHADER_TYPE(FOculusStressShadersPS, Global);
-public:
-
-	FOculusStressShadersPS() {}
-
-	explicit FOculusStressShadersPS(const ShaderMetaType::CompiledShaderInitializerType& Initializer) :
-		FGlobalShader(Initializer)
-	{
-		// Bind shader's params
-		TextureParameter.Bind(Initializer.ParameterMap, TEXT("TextureParameter"));
-	}
-
-	static bool ShouldCompilePermutation(const FGlobalShaderPermutationParameters& Parameters)
-	{
-		return !IsMetalMobilePlatform(Parameters.Platform);
-	}
-
-	virtual bool Serialize(FArchive& Ar) override
-	{
-		bool bShaderHasOutdatedParams = FGlobalShader::Serialize(Ar);
-
-		Ar << TextureParameter;
-
-		return bShaderHasOutdatedParams;
-	}
-
-	//This function is required to let us bind our runtime surface to the shader using an SRV.
-	void SetSurfaces(FRHICommandList& RHICmdList, FShaderResourceViewRHIRef TextureParameterSRV)
-	{
-		FRHIPixelShader* PixelShaderRHI = RHICmdList.GetBoundPixelShader();
-
-		if (TextureParameter.IsBound())
-		{
-			RHICmdList.SetShaderResourceViewParameter(PixelShaderRHI, TextureParameter.GetBaseIndex(), TextureParameterSRV);
-		}
-	}
-
-	//This function is required to bind our constant / uniform buffers to the shader.
-	void SetUniformBuffers(FRHICommandList& RHICmdList, FOculusPixelShaderConstantParameters& ConstantParameters, FOculusPixelShaderVariableParameters& VariableParameters)
-	{
-		FOculusPixelShaderConstantParametersRef ConstantParametersBuffer;
-		FOculusPixelShaderVariableParametersRef VariableParametersBuffer;
-
-		ConstantParametersBuffer = FOculusPixelShaderConstantParametersRef::CreateUniformBufferImmediate(ConstantParameters, UniformBuffer_SingleDraw);
-		VariableParametersBuffer = FOculusPixelShaderVariableParametersRef::CreateUniformBufferImmediate(VariableParameters, UniformBuffer_SingleDraw);
-
-		SetUniformBufferParameter(RHICmdList, RHICmdList.GetBoundPixelShader(), GetUniformBufferParameter<FOculusPixelShaderConstantParameters>(), ConstantParametersBuffer);
-		SetUniformBufferParameter(RHICmdList, RHICmdList.GetBoundPixelShader(), GetUniformBufferParameter<FOculusPixelShaderVariableParameters>(), VariableParametersBuffer);
-	}
-
-	//This is used to clean up the buffer binds after each invocation to let them be changed and used elsewhere if needed.
-	void UnbindBuffers(FRHICommandList& RHICmdList)
-	{
-		FRHIPixelShader* PixelShaderRHI = RHICmdList.GetBoundPixelShader();
-
-		if (TextureParameter.IsBound())
-		{
-			RHICmdList.SetShaderResourceViewParameter(PixelShaderRHI, TextureParameter.GetBaseIndex(), nullptr);
-		}
-	}
-
-private:
-	//This is how you declare resources that are going to be made available in the HLSL
-	LAYOUT_FIELD(FShaderResourceParameter, TextureParameter)
-};
-
-
-IMPLEMENT_SHADER_TYPE(, FOculusStressShadersPS, TEXT("/Plugin/OculusVR/Private/OculusStressShaders.usf"), TEXT("MainPixelShader"), SF_Pixel);
-
-#endif //0
-// END DO NOT MERGE -- Removing Packaging Issue
-
 //-------------------------------------------------------------------------------------------------
 // FStressTester
 //-------------------------------------------------------------------------------------------------
@@ -326,67 +244,6 @@ void FStressTester::DoTickCPU_GameThread(FOculusHMD* pPlugin)
 	}
 }
 
-void FStressTester::DoTickGPU_RenderThread(FRHICommandListImmediate& RHICmdList, class FRHITexture2D* BackBuffer, class FRHITexture2D* SrcTexture)
-{
-// DO NOT MERGE -- Removing Packaging Issue
-#if 0
-	CheckInRenderThread();
-
-	if (Mode & STM_GPU)
-	{
-		SCOPED_DRAW_EVENT(RHICmdList, StressGPURendering);
-		SCOPE_CYCLE_COUNTER(STAT_GPUStressRendering);
-
-		// Draw a quad with our shader
-		FOculusPixelShaderConstantParameters ConstantParameters;
-		FOculusPixelShaderVariableParameters VariableParameters;
-		if (GPUIterationsMultiplier <= 0)
-		{
-			VariableParameters.IterationsMultiplier = int(uint64(FPlatformTime::Seconds()*1000) % 20 + 1);
-		}
-		else
-		{
-			VariableParameters.IterationsMultiplier = GPUIterationsMultiplier;
-		}
-
-		//This is where the magic happens
-		FRHIRenderPassInfo RPInfo(BackBuffer, ERenderTargetActions::Load_Store);
-		RHICmdList.BeginRenderPass(RPInfo, TEXT("FSPass"));
-		{
-			FGraphicsPipelineStateInitializer GraphicsPSOInit;
-			RHICmdList.ApplyCachedRenderTargets(GraphicsPSOInit);
-			GraphicsPSOInit.BlendState = TStaticBlendState<>::GetRHI();
-			GraphicsPSOInit.RasterizerState = TStaticRasterizerState<>::GetRHI();
-			GraphicsPSOInit.DepthStencilState = TStaticDepthStencilState<false, CF_Always>::GetRHI();
-
-			const auto FeatureLevel = GMaxRHIFeatureLevel;
-			TShaderMapRef<FOculusVertexShader> VertexShader(GetGlobalShaderMap(FeatureLevel));
-			TShaderMapRef<FOculusStressShadersPS> PixelShader(GetGlobalShaderMap(FeatureLevel));
-
-			GraphicsPSOInit.BoundShaderState.VertexDeclarationRHI = GOculusTextureVertexDeclaration.VertexDeclarationRHI;
-			GraphicsPSOInit.BoundShaderState.VertexShaderRHI = VertexShader.GetVertexShader();
-			GraphicsPSOInit.BoundShaderState.PixelShaderRHI = PixelShader.GetPixelShader();
-			GraphicsPSOInit.PrimitiveType = PT_TriangleStrip;
-
-			SetGraphicsPipelineState(RHICmdList, GraphicsPSOInit, 0);
-
-			FShaderResourceViewRHIRef TextureParameterSRV = RHICreateShaderResourceView(SrcTexture, 0);
-			PixelShader->SetSurfaces(RHICmdList, TextureParameterSRV);
-			PixelShader->SetUniformBuffers(RHICmdList, ConstantParameters, VariableParameters);
-
-			// Draw a fullscreen quad that we can run our pixel shader on
-			RHICmdList.SetStreamSource(0, CreateTempOcculusVertexBuffer(), 0);
-			RHICmdList.DrawPrimitive(PT_TriangleStrip, 0, 2, 1);
-		}
-		RHICmdList.EndRenderPass();
-
-		PixelShader->UnbindBuffers(RHICmdList);
-	}
-
-#endif // 0
-// END DO NOT MERGE -- Removing Packaging Issue
-}
-
 
 //-------------------------------------------------------------------------------------------------
 // Console commands for managing the stress tester:
diff --git a/Engine/Plugins/Runtime/Oculus/OculusVR/Source/OculusHMD/Private/OculusHMD_StressTester.h b/Engine/Plugins/Runtime/Oculus/OculusVR/Source/OculusHMD/Private/OculusHMD_StressTester.h
--- a/Engine/Plugins/Runtime/Oculus/OculusVR/Source/OculusHMD/Private/OculusHMD_StressTester.h
+++ b/Engine/Plugins/Runtime/Oculus/OculusVR/Source/OculusHMD/Private/OculusHMD_StressTester.h
@@ -64,19 +64,8 @@ public:
 		}
 	}
 
-	static void TickGPU_RenderThread(FRHICommandListImmediate& RHICmdList, class FRHITexture2D* BackBuffer, class FRHITexture2D* SrcTexture)
-	{
-		CheckInRenderThread();
-
-		if (SharedInstance.IsValid())
-		{
-			SharedInstance->DoTickGPU_RenderThread(RHICmdList, BackBuffer, SrcTexture);
-		}
-	}
-
 protected:
 	void DoTickCPU_GameThread(class FOculusHMD* pPlugin);
-	void DoTickGPU_RenderThread(FRHICommandListImmediate& RHICmdList, class FRHITexture2D* BackBuffer, class FRHITexture2D* SrcTexture);
 
 	FStressTester();
 
