Clear explanations of the different Oculus XR runtime settings, and wording indicating which ones Epic recommends and is actively developing.

diff --git a/Engine/Plugins/Runtime/Oculus/OculusVR/Source/OculusHMD/Public/OculusFunctionLibrary.h b/Engine/Plugins/Runtime/Oculus/OculusVR/Source/OculusHMD/Public/OculusFunctionLibrary.h
--- a/Engine/Plugins/Runtime/Oculus/OculusVR/Source/OculusHMD/Public/OculusFunctionLibrary.h
+++ b/Engine/Plugins/Runtime/Oculus/OculusVR/Source/OculusHMD/Public/OculusFunctionLibrary.h
@@ -159,9 +159,11 @@ enum class EOculusDeviceType : uint8
 UENUM(BlueprintType)
 enum class EOculusXrApi : uint8
 {
-	LegacyOVRPlugin = 0 UMETA(DisplayName = "Legacy OVRPlugin", ToolTip = "Standard OVRPlugin. If unsure, this is recommended."),
-	OVRPluginOpenXR = 1 UMETA(DisplayName = "OVRPlugin OpenXR", ToolTip = "OVRPlugin using an OpenXR backend. Experimental."),
-	NativeOpenXR = 2 UMETA(DisplayName = "Native OpenXR", ToolTip = "Disable OVRPlugin in favor of the native OpenXR implementation. Must enable the OpenXR plugin."),
+	LegacyOVRPlugin = 0 UMETA(DisplayName = "Legacy Oculus SDK (no longer developed by Epic)", ToolTip = "Legacy Oculus SDK. Epic is no longer developing for this SDK, and it should only be used if there are features required for a project that are not yet supported through OpenXR, but it's hard to guarantee potential bugs will be fixed as it's not in active development at Epic. Epic recommends Native OpenXR instead, as that will be the main development focus going forward."),
+	
+	OVRPluginOpenXR = 1 UMETA(DisplayName = "Legacy Oculus SDK + OpenXR (temporary experimental solution)", ToolTip = "Legacy Oculus SDK using an OpenXR backend. Experimental. Epic recommends Native OpenXR instead, as that will be the main development focus going forward."),
+	
+	NativeOpenXR = 2 UMETA(DisplayName = "Native OpenXR with Oculus vendor extensions (Epic's development focus)", ToolTip = "Disable Legacy Oculus in favor of the native OpenXR implementation, with Oculus vendor extensions. Must enable the OpenXR plugin. This will be where Epic focuses XR development going forward. Oculus OpenXR extensions may be moved into a separate plugin (or plugins) in the future to improve modularity. The features supported by OpenXR are listed in the OpenXR specification on khronos.org, and the features supported by a given runtime can be verified with the \"OpenXR Explorer\" application on GitHub."),
 };
 
 /*
