Fix issue with losing Android force feedback

#android

diff --git a/Engine/Plugins/Runtime/Oculus/OculusVR/Source/OculusInput/Private/OculusInput.cpp b/Engine/Plugins/Runtime/Oculus/OculusVR/Source/OculusInput/Private/OculusInput.cpp
--- a/Engine/Plugins/Runtime/Oculus/OculusVR/Source/OculusInput/Private/OculusInput.cpp
+++ b/Engine/Plugins/Runtime/Oculus/OculusVR/Source/OculusInput/Private/OculusInput.cpp
@@ -876,6 +876,33 @@ void FOculusInput::SetChannelValues( int32 ControllerId, const FForceFeedbackVal
 	}
 }
 
+bool FOculusInput::SupportsForceFeedback(int32 ControllerId)
+{
+	for (FOculusControllerPair& ControllerPair : ControllerPairs)
+	{
+		if (ControllerPair.UnrealControllerIndex == ControllerId)
+		{
+			const FOculusTouchControllerState& ControllerStateLeft = ControllerPair.TouchControllerStates[(int32)EControllerHand::Left];
+			const FOculusTouchControllerState& ControllerStateRight = ControllerPair.TouchControllerStates[(int32)EControllerHand::Right];
+
+			if (!(ControllerStateLeft.bIsConnected || ControllerStateRight.bIsConnected))
+			{
+				// neither hand connected, won't be receiving force feedback
+				continue;
+			}
+
+			if (IOculusHMDModule::IsAvailable() && FOculusHMDModule::GetPluginWrapper().GetInitialized())
+			{
+				// available so could receive feedback
+				return true;
+			}
+		}
+	}
+
+	// not handling force feedback
+	return false;
+}
+
 void FOculusInput::UpdateForceFeedback( const FOculusControllerPair& ControllerPair, const EControllerHand Hand )
 {
 	const FOculusTouchControllerState& ControllerState = ControllerPair.TouchControllerStates[ (int32)Hand ];
diff --git a/Engine/Plugins/Runtime/Oculus/OculusVR/Source/OculusInput/Private/OculusInput.h b/Engine/Plugins/Runtime/Oculus/OculusVR/Source/OculusInput/Private/OculusInput.h
--- a/Engine/Plugins/Runtime/Oculus/OculusVR/Source/OculusInput/Private/OculusInput.h
+++ b/Engine/Plugins/Runtime/Oculus/OculusVR/Source/OculusInput/Private/OculusInput.h
@@ -54,6 +54,7 @@ public:
 	virtual bool Exec( UWorld* InWorld, const TCHAR* Cmd, FOutputDevice& Ar ) override;
 	virtual void SetChannelValue( int32 ControllerId, FForceFeedbackChannelType ChannelType, float Value ) override;
 	virtual void SetChannelValues( int32 ControllerId, const FForceFeedbackValues& Values ) override;
+	virtual bool SupportsForceFeedback(int32 ControllerId) override;
 
 	// IMotionController overrides
 	virtual FName GetMotionControllerDeviceTypeName() const override;
diff --git a/Engine/Plugins/Runtime/OpenXR/Source/OpenXRInput/Private/OpenXRInput.cpp b/Engine/Plugins/Runtime/OpenXR/Source/OpenXRInput/Private/OpenXRInput.cpp
--- a/Engine/Plugins/Runtime/OpenXR/Source/OpenXRInput/Private/OpenXRInput.cpp
+++ b/Engine/Plugins/Runtime/OpenXR/Source/OpenXRInput/Private/OpenXRInput.cpp
@@ -742,6 +742,17 @@ void FOpenXRInputPlugin::FOpenXRInput::SetChannelValues(int32 ControllerId, cons
 		rightHaptics);
 }
 
+bool FOpenXRInputPlugin::FOpenXRInput::SupportsForceFeedback(int32 ControllerId)
+{
+	if (!bActionsBound || OpenXRHMD == nullptr)
+	{
+		return false;
+	}
+
+	XrSession Session = OpenXRHMD->GetSession();
+	return Session != XR_NULL_HANDLE;
+}
+
 void FOpenXRInputPlugin::FOpenXRInput::SetDeviceProperty(int32 ControllerId, const FInputDeviceProperty* Property)
 {
 	if (OpenXRHMD == nullptr)
diff --git a/Engine/Plugins/Runtime/OpenXR/Source/OpenXRInput/Private/OpenXRInput.h b/Engine/Plugins/Runtime/OpenXR/Source/OpenXRInput/Private/OpenXRInput.h
--- a/Engine/Plugins/Runtime/OpenXR/Source/OpenXRInput/Private/OpenXRInput.h
+++ b/Engine/Plugins/Runtime/OpenXR/Source/OpenXRInput/Private/OpenXRInput.h
@@ -71,6 +71,7 @@ public:
 		virtual bool Exec(UWorld* InWorld, const TCHAR* Cmd, FOutputDevice& Ar) override;
 		virtual void SetChannelValue(int32 ControllerId, FForceFeedbackChannelType ChannelType, float Value) override;
 		virtual void SetChannelValues(int32 ControllerId, const FForceFeedbackValues &values) override;
+		virtual bool SupportsForceFeedback(int32 ControllerId) override;
 		virtual void SetDeviceProperty(int32 ControllerId, const FInputDeviceProperty* Property) override;
 		// IMotionController overrides
 		virtual FName GetMotionControllerDeviceTypeName() const override;
diff --git a/Engine/Plugins/Runtime/OpenXRHandTracking/Source/OpenXRHandTracking/Private/OpenXRHandTracking.h b/Engine/Plugins/Runtime/OpenXRHandTracking/Source/OpenXRHandTracking/Private/OpenXRHandTracking.h
--- a/Engine/Plugins/Runtime/OpenXRHandTracking/Source/OpenXRHandTracking/Private/OpenXRHandTracking.h
+++ b/Engine/Plugins/Runtime/OpenXRHandTracking/Source/OpenXRHandTracking/Private/OpenXRHandTracking.h
@@ -83,6 +83,7 @@ public:
 	virtual bool Exec(UWorld* InWorld, const TCHAR* Cmd, FOutputDevice& Ar) override;
 	virtual void SetChannelValue(int32 ControllerId, FForceFeedbackChannelType ChannelType, float Value) override {};
 	virtual void SetChannelValues(int32 ControllerId, const FForceFeedbackValues &values) override {};
+	virtual bool SupportsForceFeedback(int32 ControllerId) override { return false; }
 	virtual bool IsGamepadAttached() const override;
 
 	/** IHandTracker */
