UE-147503 OpenXR: Threading ensures during VRPIE startup
-We started hitting some game thread asserts because the timing of render thread destruction/creation for VRPreview changed such that some code that used to run while the RenderThread did not exist, and thus ran on the game thread, now runs on the render thread.
-Changed CalculateRenderTargetSize assert so that it tolerates Game or Rendering thread.  The HMD methods that it calls are already pipelined on some platforms.  Fixing others.  Ignoring some that are on their way out.

diff --git a/Engine/Plugins/Runtime/Oculus/OculusVR/Source/OculusHMD/Private/OculusHMD.cpp b/Engine/Plugins/Runtime/Oculus/OculusVR/Source/OculusHMD/Private/OculusHMD.cpp
--- a/Engine/Plugins/Runtime/Oculus/OculusVR/Source/OculusHMD/Private/OculusHMD.cpp
+++ b/Engine/Plugins/Runtime/Oculus/OculusVR/Source/OculusHMD/Private/OculusHMD.cpp
@@ -1238,8 +1238,14 @@ namespace OculusHMD
 
 	float FOculusHMD::GetPixelDenity() const
 	{
-		CheckInGameThread();
-		return Settings->PixelDensity;
+		if (IsInGameThread())
+		{
+			return Settings.IsValid() ? Settings->PixelDensity : 1.0f;
+		}
+		else
+		{
+			return Settings_RenderThread.IsValid() ? Settings_RenderThread->PixelDensity : 1.0f;
+		}
 	}
 
 	void FOculusHMD::SetPixelDensity(const float NewPixelDensity)
@@ -1250,8 +1256,14 @@ namespace OculusHMD
 
 	FIntPoint FOculusHMD::GetIdealRenderTargetSize() const
 	{
-		CheckInGameThread();
-		return Settings->RenderTargetSize;
+		if (IsInGameThread())
+		{
+			return Settings.IsValid() ? Settings->RenderTargetSize : 1.0f;
+		}
+		else
+		{
+			return Settings_RenderThread.IsValid() ? Settings_RenderThread->RenderTargetSize : 1.0f;
+		}
 	}
 
 	bool FOculusHMD::IsStereoEnabled() const
diff --git a/Engine/Source/Runtime/HeadMountedDisplay/Private/XRRenderTargetManager.cpp b/Engine/Source/Runtime/HeadMountedDisplay/Private/XRRenderTargetManager.cpp
--- a/Engine/Source/Runtime/HeadMountedDisplay/Private/XRRenderTargetManager.cpp
+++ b/Engine/Source/Runtime/HeadMountedDisplay/Private/XRRenderTargetManager.cpp
@@ -10,7 +10,7 @@
 
 void FXRRenderTargetManager::CalculateRenderTargetSize(const class FViewport& Viewport, uint32& InOutSizeX, uint32& InOutSizeY)
 {
-	check(IsInGameThread());
+	check(IsInGameThread() || IsInRenderingThread());
 
 	if (GEngine && GEngine->XRSystem.IsValid())
 	{
diff --git a/Engine/Source/Runtime/HeadMountedDisplay/Public/IHeadMountedDisplay.h b/Engine/Source/Runtime/HeadMountedDisplay/Public/IHeadMountedDisplay.h
--- a/Engine/Source/Runtime/HeadMountedDisplay/Public/IHeadMountedDisplay.h
+++ b/Engine/Source/Runtime/HeadMountedDisplay/Public/IHeadMountedDisplay.h
@@ -158,7 +158,7 @@ public:
 	/**
 	 * Gets the current pixel density setting.
 	 */
-	virtual float GetPixelDenity() const { return 1.0f; }
+	virtual float GetPixelDenity() const { check(IsInGameThread() || IsInRenderingThread()); return 1.0f; }
 
 	/**
 	 * Sets the pixel density. This may cause render target resizing.
@@ -168,7 +168,7 @@ public:
 	/**
 	* Gets the ideal render target size for the device. See vr.pixeldensity description.
 	*/
-	virtual FIntPoint GetIdealRenderTargetSize() const { return FIntPoint(); }
+	virtual FIntPoint GetIdealRenderTargetSize() const { check(IsInGameThread() || IsInRenderingThread()); return FIntPoint(); }
 
 	/**
 	* Gets the ideal render target size for the debug canvas on the device.
