Fix missing condition OculusMobile_APL.xml

diff --git a/Engine/Plugins/Runtime/Oculus/OculusVR/Source/OculusHMD/OculusMobile_APL.xml b/Engine/Plugins/Runtime/Oculus/OculusVR/Source/OculusHMD/OculusMobile_APL.xml
--- a/Engine/Plugins/Runtime/Oculus/OculusVR/Source/OculusHMD/OculusMobile_APL.xml
+++ b/Engine/Plugins/Runtime/Oculus/OculusVR/Source/OculusHMD/OculusMobile_APL.xml
@@ -32,6 +32,7 @@
 		<setBoolFromProperty result="bRequiresSystemKeyboard" ini="Engine" section="/Script/OculusHMD.OculusHMDRuntimeSettings" property="bRequiresSystemKeyboard" default="false"/>
 		<setStringFromProperty result="HandTrackingSupport" ini="Engine" section="/Script/OculusHMD.OculusHMDRuntimeSettings" property="HandTrackingSupport" default="ControllersOnly"/>
 		<setStringFromProperty result="HandTrackingFrequency" ini="Engine" section="/Script/OculusHMD.OculusHMDRuntimeSettings" property="HandTrackingFrequency" default="LOW"/>
+		<setBool result="bEnableHandTracking" value="false"/>
 		<setBoolOr result="bPackageForOculusMobile" 
 			arg1="$B(bPackageForOculusQuest)" arg2="$B(bPackageForOculusQuest2)"/>
 		<if condition="bPackageForOculusMobile">
