Update Oculus Input to use the new OnController functions instead of the soon to be deprecated "ControllerId"

#rnx

diff --git a/Engine/Plugins/Runtime/Oculus/OculusVR/Source/OculusInput/Private/OculusInput.cpp b/Engine/Plugins/Runtime/Oculus/OculusVR/Source/OculusInput/Private/OculusInput.cpp
--- a/Engine/Plugins/Runtime/Oculus/OculusVR/Source/OculusInput/Private/OculusInput.cpp
+++ b/Engine/Plugins/Runtime/Oculus/OculusVR/Source/OculusInput/Private/OculusInput.cpp
@@ -9,6 +9,7 @@
 #include "Misc/CoreDelegates.h"
 #include "Features/IModularFeatures.h"
 #include "Misc/ConfigCacheIni.h"
+#include "GenericPlatform/GenericPlatformInputDeviceMapper.h"
 
 #define OVR_DEBUG_LOGGING 0
 
@@ -147,8 +148,8 @@ FOculusInput::FOculusInput( const TSharedRef< FGenericApplicationMessageHandler
 
 	FOculusControllerPair& ControllerPair = *new(ControllerPairs) FOculusControllerPair();
 
-	// @todo: Unreal controller index should be assigned to us by the engine to ensure we don't contest with other devices
-	ControllerPair.UnrealControllerIndex = 0; //???? NextUnrealControllerIndex++;
+	// TODO: Map the oculus controllers uniquely instead of using the default
+	ControllerPair.DeviceId = IPlatformInputDeviceMapper::Get().GetDefaultInputDevice();
 
 	IModularFeatures::Get().RegisterModularFeature( GetModularFeatureName(), this );
 
@@ -272,6 +273,9 @@ void FOculusInput::SendControllerEvents()
 	{
 		if (MessageHandler.IsValid() && GEngine->XRSystem->GetHMDDevice())
 		{
+			FPlatformUserId PlatUser = IPlatformInputDeviceMapper::Get().GetPrimaryPlatformUser();
+			FInputDeviceId DeviceId = IPlatformInputDeviceMapper::Get().GetDefaultInputDevice();
+			
 			OculusHMD::FOculusHMD* OculusHMD = static_cast<OculusHMD::FOculusHMD*>(GEngine->XRSystem->GetHMDDevice());
 			OculusHMD->StartGameFrame_GameThread();
 
@@ -342,21 +346,21 @@ void FOculusInput::SendControllerEvents()
 						ButtonState.bIsPressed = bButtonPressed;
 						if (ButtonState.bIsPressed)
 						{
-							OnControllerButtonPressed(ButtonState, 0, false);
+							OnControllerButtonPressed(ButtonState, PlatUser, DeviceId, false);
 
 							// Set the timer for the first repeat
 							ButtonState.NextRepeatTime = CurrentTime + InitialButtonRepeatDelay;
 						}
 						else
 						{
-							OnControllerButtonReleased(ButtonState, 0, false);
+							OnControllerButtonReleased(ButtonState, PlatUser, DeviceId, false);
 						}
 					}
 
 					// Apply key repeat, if its time for that
 					if (ButtonState.bIsPressed && ButtonState.NextRepeatTime <= CurrentTime)
 					{
-						OnControllerButtonPressed(ButtonState, 0, true);
+						OnControllerButtonPressed(ButtonState, PlatUser, DeviceId, true);
 
 						// Set the timer for the next repeat
 						ButtonState.NextRepeatTime = CurrentTime + ButtonRepeatDelay;
@@ -385,6 +389,7 @@ void FOculusInput::SendControllerEvents()
 
 				for (FOculusControllerPair& ControllerPair : ControllerPairs)
 				{
+					FPlatformUserId PlatformUser = IPlatformInputDeviceMapper::Get().GetUserForInputDevice(ControllerPair.DeviceId);
 					for (int32 HandIndex = 0; HandIndex < UE_ARRAY_COUNT(ControllerPair.TouchControllerStates); ++HandIndex)
 					{
 						FOculusTouchControllerState& State = ControllerPair.TouchControllerStates[HandIndex];
@@ -424,13 +429,13 @@ void FOculusInput::SendControllerEvents()
 							if (OvrTriggerAxis != State.TriggerAxis)
 							{
 								State.TriggerAxis = OvrTriggerAxis;
-								MessageHandler->OnControllerAnalog(bIsLeft ? EKeys::OculusTouch_Left_Trigger_Axis.GetFName() : EKeys::OculusTouch_Right_Trigger_Axis.GetFName(), ControllerPair.UnrealControllerIndex, State.TriggerAxis);
+								MessageHandler->OnControllerAnalog(bIsLeft ? EKeys::OculusTouch_Left_Trigger_Axis.GetFName() : EKeys::OculusTouch_Right_Trigger_Axis.GetFName(), PlatformUser, ControllerPair.DeviceId, State.TriggerAxis);
 							}
 
 							if (OvrGripAxis != State.GripAxis)
 							{
 								State.GripAxis = OvrGripAxis;
-								MessageHandler->OnControllerAnalog(bIsLeft ? EKeys::OculusTouch_Left_Grip_Axis.GetFName() : EKeys::OculusTouch_Right_Grip_Axis.GetFName(), ControllerPair.UnrealControllerIndex, State.GripAxis);
+								MessageHandler->OnControllerAnalog(bIsLeft ? EKeys::OculusTouch_Left_Grip_Axis.GetFName() : EKeys::OculusTouch_Right_Grip_Axis.GetFName(), PlatformUser, ControllerPair.DeviceId, State.GripAxis);
 							}
 
 							ovrpVector2f ThumbstickValue = OvrpControllerState.Thumbstick[HandIndex];
@@ -439,13 +444,13 @@ void FOculusInput::SendControllerEvents()
 							if (ThumbstickValue.x != State.ThumbstickAxes.X)
 							{
 								State.ThumbstickAxes.X = ThumbstickValue.x;
-								MessageHandler->OnControllerAnalog(bIsLeft ? EKeys::OculusTouch_Left_Thumbstick_X.GetFName() : EKeys::OculusTouch_Right_Thumbstick_X.GetFName(), ControllerPair.UnrealControllerIndex, State.ThumbstickAxes.X);
+								MessageHandler->OnControllerAnalog(bIsLeft ? EKeys::OculusTouch_Left_Thumbstick_X.GetFName() : EKeys::OculusTouch_Right_Thumbstick_X.GetFName(), PlatformUser, ControllerPair.DeviceId, State.ThumbstickAxes.X);
 							}
 
 							if (ThumbstickValue.y != State.ThumbstickAxes.Y)
 							{
 								State.ThumbstickAxes.Y = ThumbstickValue.y;
-								MessageHandler->OnControllerAnalog(bIsLeft ? EKeys::OculusTouch_Left_Thumbstick_Y.GetFName() : EKeys::OculusTouch_Right_Thumbstick_Y.GetFName(), ControllerPair.UnrealControllerIndex, State.ThumbstickAxes.Y);
+								MessageHandler->OnControllerAnalog(bIsLeft ? EKeys::OculusTouch_Left_Thumbstick_Y.GetFName() : EKeys::OculusTouch_Right_Thumbstick_Y.GetFName(), PlatformUser, ControllerPair.DeviceId, State.ThumbstickAxes.Y);
 							}
 
 							if (TouchpadValue.x != State.TouchpadAxes.X)
@@ -554,21 +559,21 @@ void FOculusInput::SendControllerEvents()
 									ButtonState.bIsPressed = bButtonPressed;
 									if (ButtonState.bIsPressed)
 									{
-										OnControllerButtonPressed(ButtonState, ControllerPair.UnrealControllerIndex, false);
+										OnControllerButtonPressed(ButtonState, PlatformUser, ControllerPair.DeviceId, false);
 
 										// Set the timer for the first repeat
 										ButtonState.NextRepeatTime = CurrentTime + InitialButtonRepeatDelay;
 									}
 									else
 									{
-										OnControllerButtonReleased(ButtonState, ControllerPair.UnrealControllerIndex, false);
+										OnControllerButtonReleased(ButtonState, PlatformUser, ControllerPair.DeviceId, false);
 									}
 								}
 
 								// Apply key repeat, if its time for that
 								if (ButtonState.bIsPressed && ButtonState.NextRepeatTime <= CurrentTime)
 								{
-									OnControllerButtonPressed(ButtonState, ControllerPair.UnrealControllerIndex, true);
+									OnControllerButtonPressed(ButtonState, PlatformUser, ControllerPair.DeviceId, true);
 
 									// Set the timer for the next repeat
 									ButtonState.NextRepeatTime = CurrentTime + ButtonRepeatDelay;
@@ -631,7 +636,7 @@ void FOculusInput::SendControllerEvents()
 
 								if (CurrentAxisVal != CapState.State)
 								{
-									MessageHandler->OnControllerAnalog(CapState.Axis, ControllerPair.UnrealControllerIndex, CurrentAxisVal);
+									MessageHandler->OnControllerAnalog(CapState.Axis, PlatformUser, ControllerPair.DeviceId, CurrentAxisVal);
 
 									CapState.State = CurrentAxisVal;
 								}
@@ -663,6 +668,7 @@ void FOculusInput::SendControllerEvents()
 			{
 				for (FOculusControllerPair& ControllerPair : ControllerPairs)
 				{
+					FPlatformUserId PlatformUser = IPlatformInputDeviceMapper::Get().GetUserForInputDevice(ControllerPair.DeviceId);
 					for (int32 HandIndex = 0; HandIndex < UE_ARRAY_COUNT(ControllerPair.HandControllerStates); ++HandIndex)
 					{
 						FOculusHandControllerState& State = ControllerPair.HandControllerStates[HandIndex];
@@ -751,11 +757,11 @@ void FOculusInput::SendControllerEvents()
 										PinchState.bIsPressed = bPressed;
 										if (PinchState.bIsPressed)
 										{
-											OnControllerButtonPressed(PinchState, ControllerPair.UnrealControllerIndex, false);
+											OnControllerButtonPressed(PinchState, PlatformUser, ControllerPair.DeviceId, false);
 										}
 										else
 										{
-											OnControllerButtonReleased(PinchState, ControllerPair.UnrealControllerIndex, false);
+											OnControllerButtonReleased(PinchState, PlatformUser, ControllerPair.DeviceId, false);
 										}
 									}
 								}
@@ -774,7 +780,7 @@ void FOculusInput::SendControllerEvents()
 
 									if (PinchValue != PinchStrength.State)
 									{
-										MessageHandler->OnControllerAnalog(PinchStrength.Axis, ControllerPair.UnrealControllerIndex, PinchValue);
+										MessageHandler->OnControllerAnalog(PinchStrength.Axis, PlatformUser, ControllerPair.DeviceId, PinchValue);
 										PinchStrength.State = PinchValue;
 									}
 								}
@@ -822,9 +828,14 @@ void FOculusInput::SetChannelValue( int32 ControllerId, FForceFeedbackChannelTyp
 {
 	const EControllerHand Hand = ( ChannelType == FForceFeedbackChannelType::LEFT_LARGE || ChannelType == FForceFeedbackChannelType::LEFT_SMALL ) ? EControllerHand::Left : EControllerHand::Right;
 
+	IPlatformInputDeviceMapper& DeviceMapper = IPlatformInputDeviceMapper::Get();
+	FPlatformUserId InPlatformUser = FGenericPlatformMisc::GetPlatformUserForUserIndex(ControllerId);
+	FInputDeviceId InDeviceId = INPUTDEVICEID_NONE;
+	DeviceMapper.RemapControllerIdToPlatformUserAndDevice(ControllerId, InPlatformUser, InDeviceId);
+	
 	for( FOculusControllerPair& ControllerPair : ControllerPairs )
 	{
-		if( ControllerPair.UnrealControllerIndex == ControllerId )
+		if( ControllerPair.DeviceId == InDeviceId )
 		{
 			FOculusTouchControllerState& ControllerState = ControllerPair.TouchControllerStates[ (int32)Hand ];
 
@@ -852,9 +863,14 @@ void FOculusInput::SetChannelValue( int32 ControllerId, FForceFeedbackChannelTyp
 
 void FOculusInput::SetChannelValues( int32 ControllerId, const FForceFeedbackValues& Values )
 {
+	IPlatformInputDeviceMapper& DeviceMapper = IPlatformInputDeviceMapper::Get();
+	FPlatformUserId InPlatformUser = FGenericPlatformMisc::GetPlatformUserForUserIndex(ControllerId);
+	FInputDeviceId InDeviceId = INPUTDEVICEID_NONE;
+	DeviceMapper.RemapControllerIdToPlatformUserAndDevice(ControllerId, InPlatformUser, InDeviceId);
+	
 	for( FOculusControllerPair& ControllerPair : ControllerPairs )
 	{
-		if( ControllerPair.UnrealControllerIndex == ControllerId )
+		if( ControllerPair.DeviceId == InDeviceId )
 		{
 			// @todo: The SMALL channel controls frequency, the LARGE channel controls amplitude.  This is a bit of a weird fit.
 			FOculusTouchControllerState& LeftControllerState = ControllerPair.TouchControllerStates[ (int32)EControllerHand::Left ];
@@ -878,9 +894,14 @@ void FOculusInput::SetChannelValues( int32 ControllerId, const FForceFeedbackVal
 
 bool FOculusInput::SupportsForceFeedback(int32 ControllerId)
 {
+	IPlatformInputDeviceMapper& DeviceMapper = IPlatformInputDeviceMapper::Get();
+	FPlatformUserId InPlatformUser = FGenericPlatformMisc::GetPlatformUserForUserIndex(ControllerId);
+	FInputDeviceId InDeviceId = INPUTDEVICEID_NONE;
+	DeviceMapper.RemapControllerIdToPlatformUserAndDevice(ControllerId, InPlatformUser, InDeviceId);
+	
 	for (FOculusControllerPair& ControllerPair : ControllerPairs)
 	{
-		if (ControllerPair.UnrealControllerIndex == ControllerId)
+		if (ControllerPair.DeviceId == InDeviceId)
 		{
 			const FOculusTouchControllerState& ControllerStateLeft = ControllerPair.TouchControllerStates[(int32)EControllerHand::Left];
 			const FOculusTouchControllerState& ControllerStateRight = ControllerPair.TouchControllerStates[(int32)EControllerHand::Right];
@@ -946,25 +967,25 @@ void FOculusInput::UpdateForceFeedback( const FOculusControllerPair& ControllerP
 	}
 }
 
-bool FOculusInput::OnControllerButtonPressed(const FOculusButtonState& ButtonState, int32 ControllerId, bool IsRepeat)
+bool FOculusInput::OnControllerButtonPressed(const FOculusButtonState& ButtonState, FPlatformUserId UserId, FInputDeviceId DeviceId, bool IsRepeat)
 {
-	bool result = MessageHandler->OnControllerButtonPressed(ButtonState.Key, ControllerId, IsRepeat);
+	bool result = MessageHandler->OnControllerButtonPressed(ButtonState.Key, UserId, DeviceId, IsRepeat);
 
 	if (!ButtonState.EmulatedKey.IsNone())
 	{
-		MessageHandler->OnControllerButtonPressed(ButtonState.EmulatedKey, ControllerId, IsRepeat);
+		MessageHandler->OnControllerButtonPressed(ButtonState.EmulatedKey, UserId, DeviceId, IsRepeat);
 	}
 
 	return result;
 }
 
-bool FOculusInput::OnControllerButtonReleased(const FOculusButtonState& ButtonState, int32 ControllerId, bool IsRepeat)
+bool FOculusInput::OnControllerButtonReleased(const FOculusButtonState& ButtonState, FPlatformUserId UserId, FInputDeviceId DeviceId, bool IsRepeat)
 {
-	bool result = MessageHandler->OnControllerButtonReleased(ButtonState.Key, ControllerId, IsRepeat);
+	bool result = MessageHandler->OnControllerButtonReleased(ButtonState.Key, UserId, DeviceId, IsRepeat);
 
 	if (!ButtonState.EmulatedKey.IsNone())
 	{
-		MessageHandler->OnControllerButtonReleased(ButtonState.EmulatedKey, ControllerId, IsRepeat);
+		MessageHandler->OnControllerButtonReleased(ButtonState.EmulatedKey, UserId, DeviceId, IsRepeat);
 	}
 
 	return result;
@@ -978,12 +999,17 @@ FName FOculusInput::GetMotionControllerDeviceTypeName() const
 
 bool FOculusInput::GetControllerOrientationAndPosition( const int32 ControllerIndex, const EControllerHand DeviceHand, FRotator& OutOrientation, FVector& OutPosition, float WorldToMetersScale) const
 {
+	IPlatformInputDeviceMapper& DeviceMapper = IPlatformInputDeviceMapper::Get();
+	FPlatformUserId InPlatformUser = FGenericPlatformMisc::GetPlatformUserForUserIndex(ControllerIndex);
+	FInputDeviceId InDeviceId = INPUTDEVICEID_NONE;
+	DeviceMapper.RemapControllerIdToPlatformUserAndDevice(ControllerIndex, InPlatformUser, InDeviceId);
+	
 	// Don't do renderthread pose update if MRC is active due to controller jitter issues with SceneCaptures
 	if (IsInGameThread() || !UOculusMRFunctionLibrary::IsMrcActive())
 	{
 		for (const FOculusControllerPair& ControllerPair : ControllerPairs)
 		{
-			if (ControllerPair.UnrealControllerIndex == ControllerIndex)
+			if (ControllerPair.DeviceId == InDeviceId)
 			{
 				if ((DeviceHand == EControllerHand::Left) || (DeviceHand == EControllerHand::Right))
 				{
@@ -1066,10 +1092,15 @@ ETrackingStatus FOculusInput::GetControllerTrackingStatus(const int32 Controller
 	{
 		return TrackingStatus;
 	}
+
+	IPlatformInputDeviceMapper& DeviceMapper = IPlatformInputDeviceMapper::Get();
+	FPlatformUserId InPlatformUser = FGenericPlatformMisc::GetPlatformUserForUserIndex(ControllerIndex);
+	FInputDeviceId InDeviceId = INPUTDEVICEID_NONE;
+	DeviceMapper.RemapControllerIdToPlatformUserAndDevice(ControllerIndex, InPlatformUser, InDeviceId);
 	
 	for( const FOculusControllerPair& ControllerPair : ControllerPairs )
 	{
-		if( ControllerPair.UnrealControllerIndex == ControllerIndex )
+		if( ControllerPair.DeviceId == InDeviceId )
 		{
 			const FOculusTouchControllerState& ControllerState = ControllerPair.TouchControllerStates[ (int32)DeviceHand ];
 			if (ControllerState.bIsConnected)
@@ -1104,9 +1135,14 @@ ETrackingStatus FOculusInput::GetControllerTrackingStatus(const int32 Controller
 
 void FOculusInput::SetHapticFeedbackValues(int32 ControllerId, int32 Hand, const FHapticFeedbackValues& Values)
 {
+	IPlatformInputDeviceMapper& DeviceMapper = IPlatformInputDeviceMapper::Get();
+	FPlatformUserId InPlatformUser = FGenericPlatformMisc::GetPlatformUserForUserIndex(ControllerId);
+	FInputDeviceId InDeviceId = INPUTDEVICEID_NONE;
+	DeviceMapper.RemapControllerIdToPlatformUserAndDevice(ControllerId, InPlatformUser, InDeviceId);
+	
 	for (FOculusControllerPair& ControllerPair : ControllerPairs)
 	{
-		if (ControllerPair.UnrealControllerIndex == ControllerId)
+		if (ControllerPair.DeviceId == InDeviceId)
 		{
 			FOculusTouchControllerState& ControllerState = ControllerPair.TouchControllerStates[Hand];
 			if (ControllerState.bIsConnected)
diff --git a/Engine/Plugins/Runtime/Oculus/OculusVR/Source/OculusInput/Private/OculusInput.h b/Engine/Plugins/Runtime/Oculus/OculusVR/Source/OculusInput/Private/OculusInput.h
--- a/Engine/Plugins/Runtime/Oculus/OculusVR/Source/OculusInput/Private/OculusInput.h
+++ b/Engine/Plugins/Runtime/Oculus/OculusVR/Source/OculusInput/Private/OculusInput.h
@@ -76,8 +76,8 @@ private:
 	/** Applies force feedback settings to the controller */
 	void UpdateForceFeedback( const FOculusControllerPair& ControllerPair, const EControllerHand Hand );
 
-	bool OnControllerButtonPressed( const FOculusButtonState& ButtonState, int32 ControllerId, bool IsRepeat );
-	bool OnControllerButtonReleased( const FOculusButtonState& ButtonState, int32 ControllerId, bool IsRepeat );
+	bool OnControllerButtonPressed( const FOculusButtonState& ButtonState, FPlatformUserId UserId, FInputDeviceId DeviceId, bool IsRepeat );
+	bool OnControllerButtonReleased( const FOculusButtonState& ButtonState, FPlatformUserId UserId, FInputDeviceId DeviceId, bool IsRepeat );
 
 private:
 
diff --git a/Engine/Plugins/Runtime/Oculus/OculusVR/Source/OculusInput/Private/OculusInputState.h b/Engine/Plugins/Runtime/Oculus/OculusVR/Source/OculusInput/Private/OculusInputState.h
--- a/Engine/Plugins/Runtime/Oculus/OculusVR/Source/OculusInput/Private/OculusInputState.h
+++ b/Engine/Plugins/Runtime/Oculus/OculusVR/Source/OculusInput/Private/OculusInputState.h
@@ -515,8 +515,8 @@ struct FOculusRemoteControllerState
 
 struct FOculusControllerPair
 {
-	/** The Unreal controller index assigned to this pair */
-	int32 UnrealControllerIndex;
+	/** The input device ID for this oculus controller */
+	FInputDeviceId DeviceId;
 
 	/** Current device state for either hand */
 	FOculusTouchControllerState TouchControllerStates[2];
@@ -524,7 +524,7 @@ struct FOculusControllerPair
 	FOculusHandControllerState HandControllerStates[2];
 
 	FOculusControllerPair()
-	  : UnrealControllerIndex(INDEX_NONE),
+	  : DeviceId(INPUTDEVICEID_NONE),
 		TouchControllerStates(),
 		HandControllerStates()
 	{
