Fix CIS errors from oculus hand tracking

#rnx

diff --git a/Engine/Plugins/Runtime/Oculus/OculusVR/Source/OculusInput/Private/OculusHandTracking.cpp b/Engine/Plugins/Runtime/Oculus/OculusVR/Source/OculusInput/Private/OculusHandTracking.cpp
--- a/Engine/Plugins/Runtime/Oculus/OculusVR/Source/OculusInput/Private/OculusHandTracking.cpp
+++ b/Engine/Plugins/Runtime/Oculus/OculusVR/Source/OculusInput/Private/OculusHandTracking.cpp
@@ -13,6 +13,7 @@
 #include "Materials/Material.h"
 #include "Materials/MaterialInstanceDynamic.h"
 #include "Model.h"
+#include "GenericPlatform/GenericPlatformInputDeviceMapper.h"
 
 #define OCULUS_TO_UE4_SCALE 100.0f
 
@@ -21,6 +22,11 @@ namespace OculusInput
 
 FQuat FOculusHandTracking::GetBoneRotation(const int32 ControllerIndex, const EOculusHandType DeviceHand, const EBone BoneId)
 {
+	IPlatformInputDeviceMapper& DeviceMapper = IPlatformInputDeviceMapper::Get();
+	FPlatformUserId InPlatformUser = FGenericPlatformMisc::GetPlatformUserForUserIndex(ControllerIndex);
+	FInputDeviceId InDeviceId = INPUTDEVICEID_NONE;
+	DeviceMapper.RemapControllerIdToPlatformUserAndDevice(ControllerIndex, InPlatformUser, InDeviceId);
+	
 	FQuat Rotation = FQuat::Identity;
 	if (BoneId <= EBone::Invalid && BoneId >= EBone::Bone_Max)
 	{
@@ -34,7 +40,7 @@ FQuat FOculusHandTracking::GetBoneRotation(const int32 ControllerIndex, const EO
 		TArray<FOculusControllerPair> ControllerPairs = OculusInputModule.Get()->ControllerPairs;
 		for (const FOculusControllerPair& HandPair : ControllerPairs)
 		{
-			if (HandPair.UnrealControllerIndex == ControllerIndex)
+			if (HandPair.DeviceId == InDeviceId)
 			{
 				if (DeviceHand != EOculusHandType::None)
 				{
@@ -55,13 +61,18 @@ FQuat FOculusHandTracking::GetBoneRotation(const int32 ControllerIndex, const EO
 float FOculusHandTracking::GetHandScale(const int32 ControllerIndex, const EOculusHandType DeviceHand)
 {
 #if OCULUS_INPUT_SUPPORTED_PLATFORMS
+	IPlatformInputDeviceMapper& DeviceMapper = IPlatformInputDeviceMapper::Get();
+	FPlatformUserId InPlatformUser = FGenericPlatformMisc::GetPlatformUserForUserIndex(ControllerIndex);
+	FInputDeviceId InDeviceId = INPUTDEVICEID_NONE;
+	DeviceMapper.RemapControllerIdToPlatformUserAndDevice(ControllerIndex, InPlatformUser, InDeviceId);
+	
 	TSharedPtr<FOculusInput> OculusInputModule = StaticCastSharedPtr<FOculusInput>(IOculusInputModule::Get().GetInputDevice());
 	if (OculusInputModule.IsValid())
 	{
 		TArray<FOculusControllerPair> ControllerPairs = OculusInputModule.Get()->ControllerPairs;
 		for (const FOculusControllerPair& HandPair : ControllerPairs)
 		{
-			if (HandPair.UnrealControllerIndex == ControllerIndex)
+			if (HandPair.DeviceId == InDeviceId)
 			{
 				if (DeviceHand != EOculusHandType::None)
 				{
@@ -78,13 +89,18 @@ float FOculusHandTracking::GetHandScale(const int32 ControllerIndex, const EOcul
 ETrackingConfidence FOculusHandTracking::GetTrackingConfidence(const int32 ControllerIndex, const EOculusHandType DeviceHand)
 {
 #if OCULUS_INPUT_SUPPORTED_PLATFORMS
+	IPlatformInputDeviceMapper& DeviceMapper = IPlatformInputDeviceMapper::Get();
+	FPlatformUserId InPlatformUser = FGenericPlatformMisc::GetPlatformUserForUserIndex(ControllerIndex);
+	FInputDeviceId InDeviceId = INPUTDEVICEID_NONE;
+	DeviceMapper.RemapControllerIdToPlatformUserAndDevice(ControllerIndex, InPlatformUser, InDeviceId);
+	
 	TSharedPtr<FOculusInput> OculusInputModule = StaticCastSharedPtr<FOculusInput>(IOculusInputModule::Get().GetInputDevice());
 	if (OculusInputModule.IsValid())
 	{
 		TArray<FOculusControllerPair> ControllerPairs = OculusInputModule.Get()->ControllerPairs;
 		for (const FOculusControllerPair& HandPair : ControllerPairs)
 		{
-			if (HandPair.UnrealControllerIndex == ControllerIndex)
+			if (HandPair.DeviceId == InDeviceId)
 			{
 				if (DeviceHand != EOculusHandType::None)
 				{
@@ -100,13 +116,18 @@ ETrackingConfidence FOculusHandTracking::GetTrackingConfidence(const int32 Contr
 
 ETrackingConfidence FOculusHandTracking::GetFingerTrackingConfidence(const int32 ControllerIndex, const EOculusHandType DeviceHand, const EOculusHandAxes Finger)
 {
+	IPlatformInputDeviceMapper& DeviceMapper = IPlatformInputDeviceMapper::Get();
+	FPlatformUserId InPlatformUser = FGenericPlatformMisc::GetPlatformUserForUserIndex(ControllerIndex);
+	FInputDeviceId InDeviceId = INPUTDEVICEID_NONE;
+	DeviceMapper.RemapControllerIdToPlatformUserAndDevice(ControllerIndex, InPlatformUser, InDeviceId);
+	
 	TSharedPtr<FOculusInput> OculusInputModule = StaticCastSharedPtr<FOculusInput>(IOculusInputModule::Get().GetInputDevice());
 	if (OculusInputModule.IsValid())
 	{
 		TArray<FOculusControllerPair> ControllerPairs = OculusInputModule.Get()->ControllerPairs;
 		for (const FOculusControllerPair& HandPair : ControllerPairs)
 		{
-			if (HandPair.UnrealControllerIndex == ControllerIndex)
+			if (HandPair.DeviceId == InDeviceId)
 			{
 				if (DeviceHand != EOculusHandType::None)
 				{
@@ -122,13 +143,18 @@ ETrackingConfidence FOculusHandTracking::GetFingerTrackingConfidence(const int32
 FTransform FOculusHandTracking::GetPointerPose(const int32 ControllerIndex, const EOculusHandType DeviceHand, const float WorldToMeters)
 {
 #if OCULUS_INPUT_SUPPORTED_PLATFORMS
+	IPlatformInputDeviceMapper& DeviceMapper = IPlatformInputDeviceMapper::Get();
+	FPlatformUserId InPlatformUser = FGenericPlatformMisc::GetPlatformUserForUserIndex(ControllerIndex);
+	FInputDeviceId InDeviceId = INPUTDEVICEID_NONE;
+	DeviceMapper.RemapControllerIdToPlatformUserAndDevice(ControllerIndex, InPlatformUser, InDeviceId);
+	
 	TSharedPtr<FOculusInput> OculusInputModule = StaticCastSharedPtr<FOculusInput>(IOculusInputModule::Get().GetInputDevice());
 	if (OculusInputModule.IsValid())
 	{
 		TArray<FOculusControllerPair> ControllerPairs = OculusInputModule.Get()->ControllerPairs;
 		for (const FOculusControllerPair& HandPair : ControllerPairs)
 		{
-			if (HandPair.UnrealControllerIndex == ControllerIndex)
+			if (HandPair.DeviceId == InDeviceId)
 			{
 				if (DeviceHand != EOculusHandType::None)
 				{
@@ -147,13 +173,18 @@ FTransform FOculusHandTracking::GetPointerPose(const int32 ControllerIndex, cons
 bool FOculusHandTracking::IsPointerPoseValid(const int32 ControllerIndex, const EOculusHandType DeviceHand)
 {
 #if OCULUS_INPUT_SUPPORTED_PLATFORMS
+	IPlatformInputDeviceMapper& DeviceMapper = IPlatformInputDeviceMapper::Get();
+	FPlatformUserId InPlatformUser = FGenericPlatformMisc::GetPlatformUserForUserIndex(ControllerIndex);
+	FInputDeviceId InDeviceId = INPUTDEVICEID_NONE;
+	DeviceMapper.RemapControllerIdToPlatformUserAndDevice(ControllerIndex, InPlatformUser, InDeviceId);
+	
 	TSharedPtr<FOculusInput> OculusInputModule = StaticCastSharedPtr<FOculusInput>(IOculusInputModule::Get().GetInputDevice());
 	if (OculusInputModule.IsValid())
 	{
 		TArray<FOculusControllerPair> ControllerPairs = OculusInputModule.Get()->ControllerPairs;
 		for (const FOculusControllerPair& HandPair : ControllerPairs)
 		{
-			if (HandPair.UnrealControllerIndex == ControllerIndex)
+			if (HandPair.DeviceId == InDeviceId)
 			{
 				if (DeviceHand != EOculusHandType::None)
 				{
@@ -181,13 +212,18 @@ bool FOculusHandTracking::IsHandTrackingEnabled()
 bool FOculusHandTracking::IsHandDominant(const int32 ControllerIndex, const EOculusHandType DeviceHand)
 {
 #if OCULUS_INPUT_SUPPORTED_PLATFORMS
+	IPlatformInputDeviceMapper& DeviceMapper = IPlatformInputDeviceMapper::Get();
+	FPlatformUserId InPlatformUser = FGenericPlatformMisc::GetPlatformUserForUserIndex(ControllerIndex);
+	FInputDeviceId InDeviceId = INPUTDEVICEID_NONE;
+	DeviceMapper.RemapControllerIdToPlatformUserAndDevice(ControllerIndex, InPlatformUser, InDeviceId);
+	
 	TSharedPtr<FOculusInput> OculusInputModule = StaticCastSharedPtr<FOculusInput>(IOculusInputModule::Get().GetInputDevice());
 	if (OculusInputModule.IsValid())
 	{
 		TArray<FOculusControllerPair> ControllerPairs = OculusInputModule.Get()->ControllerPairs;
 		for (const FOculusControllerPair& HandPair : ControllerPairs)
 		{
-			if (HandPair.UnrealControllerIndex == ControllerIndex)
+			if (HandPair.DeviceId == InDeviceId)
 			{
 				if (DeviceHand != EOculusHandType::None)
 				{
@@ -203,13 +239,18 @@ bool FOculusHandTracking::IsHandDominant(const int32 ControllerIndex, const EOcu
 
 bool FOculusHandTracking::IsHandPositionValid(int32 ControllerIndex, EOculusHandType DeviceHand)
 {
+	IPlatformInputDeviceMapper& DeviceMapper = IPlatformInputDeviceMapper::Get();
+	FPlatformUserId InPlatformUser = FGenericPlatformMisc::GetPlatformUserForUserIndex(ControllerIndex);
+	FInputDeviceId InDeviceId = INPUTDEVICEID_NONE;
+	DeviceMapper.RemapControllerIdToPlatformUserAndDevice(ControllerIndex, InPlatformUser, InDeviceId);
+	
 	TSharedPtr<FOculusInput> OculusInputModule = StaticCastSharedPtr<FOculusInput>(IOculusInputModule::Get().GetInputDevice());
 	if (OculusInputModule.IsValid())
 	{
 		TArray<FOculusControllerPair> ControllerPairs = OculusInputModule.Get()->ControllerPairs;
 		for (const FOculusControllerPair& HandPair : ControllerPairs)
 		{
-			if (HandPair.UnrealControllerIndex == ControllerIndex)
+			if (HandPair.DeviceId == InDeviceId)
 			{
 				if (DeviceHand != EOculusHandType::None)
 				{
