OculusVR deprecation : OculusOpenXRHMD module

diff --git a/Engine/Plugins/Runtime/Oculus/OculusVR/Source/OculusOpenXRHMD/Private/OculusOpenXRHMD.h b/Engine/Plugins/Runtime/Oculus/OculusVR/Source/OculusOpenXRHMD/Private/OculusOpenXRHMD.h
--- a/Engine/Plugins/Runtime/Oculus/OculusVR/Source/OculusOpenXRHMD/Private/OculusOpenXRHMD.h
+++ b/Engine/Plugins/Runtime/Oculus/OculusVR/Source/OculusOpenXRHMD/Private/OculusOpenXRHMD.h
@@ -6,6 +6,8 @@
 
 DECLARE_LOG_CATEGORY_EXTERN(LogOculusOpenXRPlugin, Log, All);
 
+PRAGMA_DISABLE_DEPRECATION_WARNINGS
+
 class FOculusOpenXRHMD : public IOculusOpenXRHMDPlugin
 {
 private:
@@ -49,4 +51,6 @@ public:
 	virtual const void* OnEndFrame(XrSession InSession, XrTime DisplayTime, const TArray<XrSwapchainSubImage> InColorImages, const TArray<XrSwapchainSubImage> InDepthImages, const void* InNext) override;
 	virtual const void* OnSyncActions(XrSession InSession, const void* InNext) override;
 	virtual void PostSyncActions(XrSession InSession) override;
-};
\ No newline at end of file
+};
+
+PRAGMA_ENABLE_DEPRECATION_WARNINGS
diff --git a/Engine/Plugins/Runtime/Oculus/OculusVR/Source/OculusOpenXRHMD/Public/IOculusOpenXRHMDPlugin.h b/Engine/Plugins/Runtime/Oculus/OculusVR/Source/OculusOpenXRHMD/Public/IOculusOpenXRHMDPlugin.h
--- a/Engine/Plugins/Runtime/Oculus/OculusVR/Source/OculusOpenXRHMD/Public/IOculusOpenXRHMDPlugin.h
+++ b/Engine/Plugins/Runtime/Oculus/OculusVR/Source/OculusOpenXRHMD/Public/IOculusOpenXRHMDPlugin.h
@@ -3,6 +3,7 @@
 #pragma once
 #include "IOpenXRExtensionPlugin.h"
 
+class UE_DEPRECATED(5.1, "OculusVR plugin is deprecated") IOculusOpenXRHMDPlugin;
 class IOculusOpenXRHMDPlugin : public IOpenXRExtensionPlugin, public IModuleInterface
 {
 	
