OculusVR deprecation: OculusEditor module

diff --git a/Engine/Plugins/Runtime/Oculus/OculusVR/Source/OculusEditor/Private/OculusEditorModule.h b/Engine/Plugins/Runtime/Oculus/OculusVR/Source/OculusEditor/Private/OculusEditorModule.h
--- a/Engine/Plugins/Runtime/Oculus/OculusVR/Source/OculusEditor/Private/OculusEditorModule.h
+++ b/Engine/Plugins/Runtime/Oculus/OculusVR/Source/OculusEditor/Private/OculusEditorModule.h
@@ -17,6 +17,8 @@ class FMenuBuilder;
 
 enum class ECheckBoxState : uint8;
 
+PRAGMA_DISABLE_DEPRECATION_WARNINGS
+
 class FOculusEditorModule : public IOculusEditorModule
 {
 public:
@@ -73,3 +75,6 @@ public:
 	bool GetBuildTelemetryCheckBoxEnabled() const;
 	EVisibility GetOculusHMDAvailableWarningVisibility() const;
 };
+
+PRAGMA_ENABLE_DEPRECATION_WARNINGS
+
diff --git a/Engine/Plugins/Runtime/Oculus/OculusVR/Source/OculusEditor/Private/OculusEditorSettings.cpp b/Engine/Plugins/Runtime/Oculus/OculusVR/Source/OculusEditor/Private/OculusEditorSettings.cpp
--- a/Engine/Plugins/Runtime/Oculus/OculusVR/Source/OculusEditor/Private/OculusEditorSettings.cpp
+++ b/Engine/Plugins/Runtime/Oculus/OculusVR/Source/OculusEditor/Private/OculusEditorSettings.cpp
@@ -2,11 +2,12 @@
 
 #include "OculusEditorSettings.h"
 
+PRAGMA_DISABLE_DEPRECATION_WARNINGS
 
-UOculusEditorSettings::UOculusEditorSettings()
+UDEPRECATED_UOculusEditorSettings::UDEPRECATED_UOculusEditorSettings()
 	: PerfToolTargetPlatform(EOculusPlatform::PC)
 {
 
 }
 
-
+PRAGMA_ENABLE_DEPRECATION_WARNINGS
diff --git a/Engine/Plugins/Runtime/Oculus/OculusVR/Source/OculusEditor/Private/OculusPlatformToolSettings.cpp b/Engine/Plugins/Runtime/Oculus/OculusVR/Source/OculusEditor/Private/OculusPlatformToolSettings.cpp
--- a/Engine/Plugins/Runtime/Oculus/OculusVR/Source/OculusEditor/Private/OculusPlatformToolSettings.cpp
+++ b/Engine/Plugins/Runtime/Oculus/OculusVR/Source/OculusEditor/Private/OculusPlatformToolSettings.cpp
@@ -2,8 +2,9 @@
 
 #include "OculusPlatformToolSettings.h"
 
+PRAGMA_DISABLE_DEPRECATION_WARNINGS
 
-UOculusPlatformToolSettings::UOculusPlatformToolSettings()
+UDEPRECATED_UOculusPlatformToolSettings::UDEPRECATED_UOculusPlatformToolSettings()
 	: OculusTargetPlatform(EOculusPlatformTarget::Rift)
 {
 	uint8 NumPlatforms = (uint8)EOculusPlatformTarget::Length;
@@ -24,4 +25,4 @@ UOculusPlatformToolSettings::UOculusPlatformToolSettings()
 	}
 }
 
-
+PRAGMA_ENABLE_DEPRECATION_WARNINGS
diff --git a/Engine/Plugins/Runtime/Oculus/OculusVR/Source/OculusEditor/Private/OculusPlatformToolWidget.cpp b/Engine/Plugins/Runtime/Oculus/OculusVR/Source/OculusEditor/Private/OculusPlatformToolWidget.cpp
--- a/Engine/Plugins/Runtime/Oculus/OculusVR/Source/OculusEditor/Private/OculusPlatformToolWidget.cpp
+++ b/Engine/Plugins/Runtime/Oculus/OculusVR/Source/OculusEditor/Private/OculusPlatformToolWidget.cpp
@@ -35,6 +35,8 @@ static bool bShowUploadDebugSymbols = false;
 
 FString SOculusPlatformToolWidget::LogText;
 
+PRAGMA_DISABLE_DEPRECATION_WARNINGS
+
 SOculusPlatformToolWidget::SOculusPlatformToolWidget()
 {
 	LogTextUpdated = false;
@@ -978,7 +980,7 @@ void SOculusPlatformToolWidget::EnableUploadButton(bool enabled)
 
 void SOculusPlatformToolWidget::LoadConfigSettings()
 {
-	PlatformSettings = GetMutableDefault<UOculusPlatformToolSettings>();
+	PlatformSettings = GetMutableDefault<UDEPRECATED_UOculusPlatformToolSettings>();
 	PlatformEnum = StaticEnum<EOculusPlatformTarget>();
 	GamepadEmulationEnum = StaticEnum<EOculusGamepadEmulation>();
 	AssetTypeEnum = StaticEnum<EOculusAssetType>();
@@ -1705,7 +1707,7 @@ FPlatformLoadRedistPackagesTask::FPlatformLoadRedistPackagesTask(FUpdateLogTextD
 
 void FPlatformLoadRedistPackagesTask::DoWork()
 {
-	UOculusPlatformToolSettings* PlatformSettings = GetMutableDefault<UOculusPlatformToolSettings>();
+	UDEPRECATED_UOculusPlatformToolSettings* PlatformSettings = GetMutableDefault<UDEPRECATED_UOculusPlatformToolSettings>();
 
 	// Check to see if the CLI exists, we need this to load avalible redist packages
 	if (!FPaths::FileExists(FPaths::ProjectContentDir() + ProjectPlatformUtilPath))
@@ -1774,4 +1776,6 @@ void FPlatformLoadRedistPackagesTask::DoWork()
 	}
 }
 
+PRAGMA_ENABLE_DEPRECATION_WARNINGS
+
 #undef LOCTEXT_NAMESPACE
\ No newline at end of file
diff --git a/Engine/Plugins/Runtime/Oculus/OculusVR/Source/OculusEditor/Private/OculusPlatformToolWidget.h b/Engine/Plugins/Runtime/Oculus/OculusVR/Source/OculusEditor/Private/OculusPlatformToolWidget.h
--- a/Engine/Plugins/Runtime/Oculus/OculusVR/Source/OculusEditor/Private/OculusPlatformToolWidget.h
+++ b/Engine/Plugins/Runtime/Oculus/OculusVR/Source/OculusEditor/Private/OculusPlatformToolWidget.h
@@ -32,6 +32,8 @@ DECLARE_DELEGATE_OneParam(FUpdateLogTextDel, FString);
 DECLARE_DELEGATE_OneParam(FSetProcessDel, FProcHandle);
 DECLARE_DELEGATE_RetVal_TwoParams(bool, FFieldValidatorDel, FString, FString&);
 
+PRAGMA_DISABLE_DEPRECATION_WARNINGS
+
 class SOculusPlatformToolWidget : public SCompoundWidget
 {
 public:
@@ -62,7 +64,7 @@ private:
 	UEnum* PlatformEnum;
 	UEnum* GamepadEmulationEnum;
 	UEnum* AssetTypeEnum;
-	UOculusPlatformToolSettings* PlatformSettings;
+	UDEPRECATED_UOculusPlatformToolSettings* PlatformSettings;
 	TArray<TSharedPtr<FString>> OculusPlatforms;
 	TArray<TSharedPtr<FString>> RiftGamepadEmulation;
 	TArray<TSharedPtr<FString>> AssetType;
@@ -219,3 +221,5 @@ protected:
 	}
 };
 
+PRAGMA_ENABLE_DEPRECATION_WARNINGS
+
diff --git a/Engine/Plugins/Runtime/Oculus/OculusVR/Source/OculusEditor/Private/OculusToolWidget.cpp b/Engine/Plugins/Runtime/Oculus/OculusVR/Source/OculusEditor/Private/OculusToolWidget.cpp
--- a/Engine/Plugins/Runtime/Oculus/OculusVR/Source/OculusEditor/Private/OculusToolWidget.cpp
+++ b/Engine/Plugins/Runtime/Oculus/OculusVR/Source/OculusEditor/Private/OculusToolWidget.cpp
@@ -22,6 +22,8 @@
 
 #define LOCTEXT_NAMESPACE "OculusToolWidget"
 
+PRAGMA_DISABLE_DEPRECATION_WARNINGS
+
 // Misc notes and known issues:
 // * I save after every change because UE wasn't prompting to save on exit, but this makes it tough for users to undo, and doesn't prompt shader rebuild. Alternatives?
 
@@ -64,7 +66,7 @@ EVisibility SOculusToolWidget::IsVisible(FName tag) const
 	const SimpleSetting* setting = SimpleSettings.Find(tag);
 	checkf(setting != NULL, TEXT("Failed to find tag %s."), *tag.ToString());
 	if(SettingIgnored(setting->tag)) return EVisibility::Collapsed;
-	UOculusEditorSettings* EditorSettings = GetMutableDefault<UOculusEditorSettings>();
+	UDEPRECATED_UOculusEditorSettings* EditorSettings = GetMutableDefault<UDEPRECATED_UOculusEditorSettings>();
 	EOculusPlatform targetPlatform = EditorSettings->PerfToolTargetPlatform;
 	 
 	if(targetPlatform == EOculusPlatform::Mobile && !((int)setting->supportMask & (int)SupportFlags::SupportMobile)) return EVisibility::Collapsed;
@@ -89,7 +91,7 @@ void SOculusToolWidget::AddSimpleSetting(TSharedRef<SVerticalBox> box, SimpleSet
 
 bool SOculusToolWidget::SettingIgnored(FName settingKey) const
 {
-	UOculusEditorSettings* EditorSettings = GetMutableDefault<UOculusEditorSettings>();
+	UDEPRECATED_UOculusEditorSettings* EditorSettings = GetMutableDefault<UDEPRECATED_UOculusEditorSettings>();
 	bool* ignoreSetting = EditorSettings->PerfToolIgnoreList.Find(settingKey);
 	return (ignoreSetting != NULL && *ignoreSetting == true);
 }
@@ -132,7 +134,7 @@ void SOculusToolWidget::RebuildLayout()
 	if (!ScrollingContainer.IsValid()) return;
 	TSharedRef<SScrollBox> scroller = ScrollingContainer.ToSharedRef();
 
-	UOculusEditorSettings* EditorSettings = GetMutableDefault<UOculusEditorSettings>();
+	UDEPRECATED_UOculusEditorSettings* EditorSettings = GetMutableDefault<UDEPRECATED_UOculusEditorSettings>();
 	uint8 initiallySelected = 0;
 	for (uint8 i = 0; i < (uint8)EOculusPlatform::Length; ++i)
 	{
@@ -326,7 +328,7 @@ void SOculusToolWidget::Construct(const FArguments& InArgs)
 	PlatformEnum = StaticEnum<EOculusPlatform>();
 	Platforms.Reset(2);
 
-	UOculusEditorSettings* EditorSettings = GetMutableDefault<UOculusEditorSettings>();
+	UDEPRECATED_UOculusEditorSettings* EditorSettings = GetMutableDefault<UDEPRECATED_UOculusEditorSettings>();
 	for (uint8 i = 0; i < (uint8)EOculusPlatform::Length; ++i)
 	{
 		Platforms.Add(MakeShareable(new FString(PlatformEnum->GetDisplayNameTextByIndex((int64)i).ToString())));
@@ -591,7 +593,7 @@ void SOculusToolWidget::OnChangePlatform(TSharedPtr<FString> ItemSelected, ESele
 	int32 idx = PlatformEnum->GetIndexByNameString(*ItemSelected);
 	if (idx != INDEX_NONE)
 	{
-		UOculusEditorSettings* EditorSettings = GetMutableDefault<UOculusEditorSettings>();
+		UDEPRECATED_UOculusEditorSettings* EditorSettings = GetMutableDefault<UDEPRECATED_UOculusEditorSettings>();
 		EditorSettings->PerfToolTargetPlatform = (EOculusPlatform)idx;
 		EditorSettings->SaveConfig();
 	}
@@ -600,7 +602,7 @@ void SOculusToolWidget::OnChangePlatform(TSharedPtr<FString> ItemSelected, ESele
 
 FReply SOculusToolWidget::IgnoreRecommendation(FName tag)
 {
-	UOculusEditorSettings* EditorSettings = GetMutableDefault<UOculusEditorSettings>();
+	UDEPRECATED_UOculusEditorSettings* EditorSettings = GetMutableDefault<UDEPRECATED_UOculusEditorSettings>();
 	EditorSettings->PerfToolIgnoreList.Add(tag, true);
 	EditorSettings->SaveConfig();
 	return FReply::Handled();
@@ -608,13 +610,13 @@ FReply SOculusToolWidget::IgnoreRecommendation(FName tag)
 
 EVisibility SOculusToolWidget::CanUnhideIgnoredRecommendations() const
 {
-	UOculusEditorSettings* EditorSettings = GetMutableDefault<UOculusEditorSettings>();
+	UDEPRECATED_UOculusEditorSettings* EditorSettings = GetMutableDefault<UDEPRECATED_UOculusEditorSettings>();
 	return EditorSettings->PerfToolIgnoreList.Num() > 0 ? EVisibility::Visible : EVisibility::Collapsed;
 }
 
 FReply SOculusToolWidget::UnhideIgnoredRecommendations()
 {
-	UOculusEditorSettings* EditorSettings = GetMutableDefault<UOculusEditorSettings>();
+	UDEPRECATED_UOculusEditorSettings* EditorSettings = GetMutableDefault<UDEPRECATED_UOculusEditorSettings>();
 	EditorSettings->PerfToolIgnoreList.Empty();
 	EditorSettings->SaveConfig();
 	RebuildLayout();
@@ -623,7 +625,7 @@ FReply SOculusToolWidget::UnhideIgnoredRecommendations()
 
 bool SOculusToolWidget::UsingForwardShading() const
 {
-	UOculusEditorSettings* EditorSettings = GetMutableDefault<UOculusEditorSettings>();
+	UDEPRECATED_UOculusEditorSettings* EditorSettings = GetMutableDefault<UDEPRECATED_UOculusEditorSettings>();
 	URendererSettings* Settings = GetMutableDefault<URendererSettings>();
 	EOculusPlatform targetPlatform = EditorSettings->PerfToolTargetPlatform;
 	return targetPlatform == EOculusPlatform::Mobile || Settings->bForwardShading;
@@ -982,7 +984,7 @@ FReply SOculusToolWidget::SelectLight(FString lightName)
 
 FReply SOculusToolWidget::IgnoreLight(FString lightName)
 {
-	UOculusEditorSettings* EditorSettings = GetMutableDefault<UOculusEditorSettings>();
+	UDEPRECATED_UOculusEditorSettings* EditorSettings = GetMutableDefault<UDEPRECATED_UOculusEditorSettings>();
 	FString lightIgnoreKey = "IgnoreLight_" + lightName;
 	EditorSettings->PerfToolIgnoreList.Add(FName(lightIgnoreKey.GetCharArray().GetData()), true);
 	EditorSettings->SaveConfig();
@@ -1019,4 +1021,6 @@ EVisibility SOculusToolWidget::SupportDashVisibility(FName tag) const
 	return Settings->bSupportsDash ? EVisibility::Collapsed : EVisibility::Visible;
 }
 
+PRAGMA_ENABLE_DEPRECATION_WARNINGS
+
 #undef LOCTEXT_NAMESPACE
diff --git a/Engine/Plugins/Runtime/Oculus/OculusVR/Source/OculusEditor/Public/IOculusEditorModule.h b/Engine/Plugins/Runtime/Oculus/OculusVR/Source/OculusEditor/Public/IOculusEditorModule.h
--- a/Engine/Plugins/Runtime/Oculus/OculusVR/Source/OculusEditor/Public/IOculusEditorModule.h
+++ b/Engine/Plugins/Runtime/Oculus/OculusVR/Source/OculusEditor/Public/IOculusEditorModule.h
@@ -13,6 +13,7 @@ class FMenuBuilder;
 //////////////////////////////////////////////////////////////////////////
 // IOculusEditorModule
 
+class UE_DEPRECATED(5.1, "OculusVR plugin is deprecated; please use the built-in OpenXR plugin or OculusXR plugin from the Marketplace.") IOculusEditorModule;
 class IOculusEditorModule : public IModuleInterface
 {
 };
diff --git a/Engine/Plugins/Runtime/Oculus/OculusVR/Source/OculusEditor/Public/OculusEditorSettings.h b/Engine/Plugins/Runtime/Oculus/OculusVR/Source/OculusEditor/Public/OculusEditorSettings.h
--- a/Engine/Plugins/Runtime/Oculus/OculusVR/Source/OculusEditor/Public/OculusEditorSettings.h
+++ b/Engine/Plugins/Runtime/Oculus/OculusVR/Source/OculusEditor/Public/OculusEditorSettings.h
@@ -6,8 +6,10 @@
 #include "UObject/NoExportTypes.h"
 #include "OculusEditorSettings.generated.h"
 
+PRAGMA_DISABLE_DEPRECATION_WARNINGS
+
 UENUM()
-enum class EOculusPlatform : uint8
+enum class UE_DEPRECATED(5.1, "OculusVR plugin is deprecated; please use the built-in OpenXR plugin or OculusXR plugin from the Marketplace.") EOculusPlatform : uint8
 {
 	PC UMETA(DisplayName="PC"),
 	Mobile UMETA(DisplayName="Mobile"),
@@ -17,13 +19,13 @@ enum class EOculusPlatform : uint8
 /**
  * 
  */
-UCLASS(config=Editor)
-class OCULUSEDITOR_API UOculusEditorSettings : public UObject
+UCLASS(config=Editor, deprecated, meta = (DeprecationMessage = "OculusVR plugin is deprecated; please use the built-in OpenXR plugin or OculusXR plugin from the Marketplace."))
+class OCULUSEDITOR_API UDEPRECATED_UOculusEditorSettings : public UObject
 {
 	GENERATED_BODY()
 
 public:
-	UOculusEditorSettings();
+	UDEPRECATED_UOculusEditorSettings();
 
 	UPROPERTY(config, EditAnywhere, Category = Oculus)
 	TMap<FName, bool> PerfToolIgnoreList;
@@ -34,3 +36,5 @@ public:
 	UPROPERTY(globalconfig, EditAnywhere, Category = Oculus)
 	bool bAddMenuOption;
 };
+
+PRAGMA_ENABLE_DEPRECATION_WARNINGS
diff --git a/Engine/Plugins/Runtime/Oculus/OculusVR/Source/OculusEditor/Public/OculusPlatformToolSettings.h b/Engine/Plugins/Runtime/Oculus/OculusVR/Source/OculusEditor/Public/OculusPlatformToolSettings.h
--- a/Engine/Plugins/Runtime/Oculus/OculusVR/Source/OculusEditor/Public/OculusPlatformToolSettings.h
+++ b/Engine/Plugins/Runtime/Oculus/OculusVR/Source/OculusEditor/Public/OculusPlatformToolSettings.h
@@ -6,8 +6,10 @@
 #include "UObject/NoExportTypes.h"
 #include "OculusPlatformToolSettings.generated.h"
 
+PRAGMA_DISABLE_DEPRECATION_WARNINGS
+
 UENUM()
-enum class EOculusPlatformTarget : uint8
+enum class UE_DEPRECATED(5.1, "OculusVR plugin is deprecated; please use the built-in OpenXR plugin or OculusXR plugin from the Marketplace.") EOculusPlatformTarget : uint8
 {
 	Rift UMETA(DisplayName="Rift"),
 	Quest UMETA(DisplayName="Quest"),
@@ -15,7 +17,7 @@ enum class EOculusPlatformTarget : uint8
 };
 
 UENUM()
-enum class EOculusGamepadEmulation : uint8
+enum class UE_DEPRECATED(5.1, "OculusVR plugin is deprecated; please use the built-in OpenXR plugin or OculusXR plugin from the Marketplace.") EOculusGamepadEmulation : uint8
 {
 	Off UMETA(DisplayName="Off"),
 	Twinstick UMETA(DisplayName = "Twinstick"),
@@ -25,7 +27,7 @@ enum class EOculusGamepadEmulation : uint8
 };
 
 UENUM()
-enum class EOculusAssetType : uint8
+enum class UE_DEPRECATED(5.1, "OculusVR plugin is deprecated; please use the built-in OpenXR plugin or OculusXR plugin from the Marketplace.") EOculusAssetType : uint8
 {
 	Default UMETA(DisplayName="Default"),
 	Store UMETA(DisplayName="Store"),
@@ -33,7 +35,7 @@ enum class EOculusAssetType : uint8
 	Length UMETA(DisplayName="Invlaid"),
 };
 
-USTRUCT()
+USTRUCT(meta = (Deprecated = "5.1"))
 struct FRedistPackage
 {
 	GENERATED_USTRUCT_BODY()
@@ -48,7 +50,7 @@ struct FRedistPackage
 	FString Id;
 };
 
-USTRUCT()
+USTRUCT(meta = (Deprecated = "5.1"))
 struct FAssetConfig
 {
 	GENERATED_USTRUCT_BODY()
@@ -66,7 +68,7 @@ struct FAssetConfig
 	FString Sku;
 };
 
-USTRUCT()
+USTRUCT(meta = (Deprecated = "5.1"))
 struct FAssetConfigArray
 {
 	GENERATED_USTRUCT_BODY()
@@ -78,13 +80,13 @@ struct FAssetConfigArray
 /**
  * 
  */
-UCLASS(config=Editor)
-class OCULUSEDITOR_API UOculusPlatformToolSettings : public UObject
+UCLASS(config=Editor, deprecated, meta = (DeprecationMessage = "OculusVR plugin is deprecated; please use the built-in OpenXR plugin or OculusXR plugin from the Marketplace."))
+class OCULUSEDITOR_API UDEPRECATED_UOculusPlatformToolSettings : public UObject
 {
 	GENERATED_BODY()
 
 public:
-	UOculusPlatformToolSettings();
+	UDEPRECATED_UOculusPlatformToolSettings();
 
 	uint8 GetTargetPlatform()
 	{
@@ -205,34 +207,34 @@ public:
 		return (uint8)OculusTargetPlatform < OculusAssetConfigs.Num() ? &OculusAssetConfigs[(uint8)OculusTargetPlatform].ConfigArray : NULL;
 	}
 
-	UPROPERTY(config, EditAnywhere, Category = Oculus)
+	UPROPERTY(config, EditAnywhere, Category = Oculus, meta = (DeprecatedProperty))
 	FString OculusRiftBuildDirectory;
 
-	UPROPERTY(config, EditAnywhere, Category = Oculus)
+	UPROPERTY(config, EditAnywhere, Category = Oculus, meta = (DeprecatedProperty))
 	FString OculusRiftBuildVersion;
 
-	UPROPERTY(config, EditAnywhere, Category = Oculus)
+	UPROPERTY(config, EditAnywhere, Category = Oculus, meta = (DeprecatedProperty))
 	FString OculusRiftLaunchParams;
 
-	UPROPERTY(config, EditAnywhere, Category = Oculus)
+	UPROPERTY(config, EditAnywhere, Category = Oculus, meta = (DeprecatedProperty))
 	bool OculusRiftFireWallException;
 
-	UPROPERTY(config, EditAnywhere, Category = Oculus)
+	UPROPERTY(config, EditAnywhere, Category = Oculus, meta = (DeprecatedProperty))
 	FString OculusRift2DLaunchPath;
 
-	UPROPERTY(config, EditAnywhere, Category = Oculus)
+	UPROPERTY(config, EditAnywhere, Category = Oculus, meta = (DeprecatedProperty))
 	FString OculusRift2DLaunchParams;
 
-	UPROPERTY(config, EditAnywhere, Category = Oculus)
+	UPROPERTY(config, EditAnywhere, Category = Oculus, meta = (DeprecatedProperty))
 	TArray<FRedistPackage> OculusRedistPackages;
 
-	UPROPERTY(config, EditAnywhere, Category = Oculus)
+	UPROPERTY(config, EditAnywhere, Category = Oculus, meta = (DeprecatedProperty))
 	bool UploadDebugSymbols;
 
-	UPROPERTY(config, EditAnywhere, Category = Oculus)
+	UPROPERTY(config, EditAnywhere, Category = Oculus, meta = (DeprecatedProperty))
 	bool DebugSymbolsOnly;
 
-	UPROPERTY(config, EditAnywhere, Category = Oculus)
+	UPROPERTY(config, EditAnywhere, Category = Oculus, meta = (DeprecatedProperty))
 	FString BuildID;
 
 private:
@@ -269,3 +271,6 @@ private:
 	UPROPERTY(config, EditAnywhere, Category = Oculus)
 	TArray<FAssetConfigArray> OculusAssetConfigs;
 };
+
+PRAGMA_ENABLE_DEPRECATION_WARNINGS
+
