OculusVR deprecation: OculusMR module

diff --git a/Engine/Plugins/Runtime/Oculus/OculusVR/Source/OculusInput/Private/OculusInput.cpp b/Engine/Plugins/Runtime/Oculus/OculusVR/Source/OculusInput/Private/OculusInput.cpp
--- a/Engine/Plugins/Runtime/Oculus/OculusVR/Source/OculusInput/Private/OculusInput.cpp
+++ b/Engine/Plugins/Runtime/Oculus/OculusVR/Source/OculusInput/Private/OculusInput.cpp
@@ -15,6 +15,8 @@
 
 #define LOCTEXT_NAMESPACE "OculusInput"
 
+PRAGMA_DISABLE_DEPRECATION_WARNINGS
+
 namespace OculusInput
 {
 
@@ -1005,7 +1007,7 @@ bool FOculusInput::GetControllerOrientationAndPosition( const int32 ControllerIn
 	DeviceMapper.RemapControllerIdToPlatformUserAndDevice(ControllerIndex, InPlatformUser, InDeviceId);
 	
 	// Don't do renderthread pose update if MRC is active due to controller jitter issues with SceneCaptures
-	if (IsInGameThread() || !UOculusMRFunctionLibrary::IsMrcActive())
+	if (IsInGameThread() || !UDEPRECATED_UOculusMRFunctionLibrary::IsMrcActive())
 	{
 		for (const FOculusControllerPair& ControllerPair : ControllerPairs)
 		{
@@ -1390,5 +1392,7 @@ uint32 FOculusInput::GetNumberOfHandControllers() const
 
 } // namespace OculusInput
 
+PRAGMA_ENABLE_DEPRECATION_WARNINGS
+
 #undef LOCTEXT_NAMESPACE
 #endif	 // OCULUS_INPUT_SUPPORTED_PLATFORMS
diff --git a/Engine/Plugins/Runtime/Oculus/OculusVR/Source/OculusMR/Private/OculusMRFunctionLibrary.cpp b/Engine/Plugins/Runtime/Oculus/OculusVR/Source/OculusMR/Private/OculusMRFunctionLibrary.cpp
--- a/Engine/Plugins/Runtime/Oculus/OculusVR/Source/OculusMR/Private/OculusMRFunctionLibrary.cpp
+++ b/Engine/Plugins/Runtime/Oculus/OculusVR/Source/OculusMR/Private/OculusMRFunctionLibrary.cpp
@@ -15,12 +15,14 @@
 // UOculusFunctionLibrary
 //-------------------------------------------------------------------------------------------------
 
-UOculusMRFunctionLibrary::UOculusMRFunctionLibrary(const FObjectInitializer& ObjectInitializer)
+PRAGMA_DISABLE_DEPRECATION_WARNINGS
+
+UDEPRECATED_UOculusMRFunctionLibrary::UDEPRECATED_UOculusMRFunctionLibrary(const FObjectInitializer& ObjectInitializer)
 	: Super(ObjectInitializer)
 {
 }
 
-void UOculusMRFunctionLibrary::GetAllTrackedCamera(TArray<FTrackedCamera>& TrackedCameras, bool bCalibratedOnly)
+void UDEPRECATED_UOculusMRFunctionLibrary::GetAllTrackedCamera(TArray<FTrackedCamera>& TrackedCameras, bool bCalibratedOnly)
 {
 	TrackedCameras.Empty();
 
@@ -77,7 +79,7 @@ void UOculusMRFunctionLibrary::GetAllTrackedCamera(TArray<FTrackedCamera>& Track
 	}
 }
 
-OculusHMD::FOculusHMD* UOculusMRFunctionLibrary::GetOculusHMD()
+OculusHMD::FOculusHMD* UDEPRECATED_UOculusMRFunctionLibrary::GetOculusHMD()
 {
 #if OCULUS_HMD_SUPPORTED_PLATFORMS
 	if (GEngine && GEngine->XRSystem.IsValid())
@@ -92,7 +94,7 @@ OculusHMD::FOculusHMD* UOculusMRFunctionLibrary::GetOculusHMD()
 	return nullptr;
 }
 
-bool UOculusMRFunctionLibrary::GetTrackingReferenceLocationAndRotationInWorldSpace(USceneComponent* TrackingReferenceComponent, FVector& TRLocation, FRotator& TRRotation)
+bool UDEPRECATED_UOculusMRFunctionLibrary::GetTrackingReferenceLocationAndRotationInWorldSpace(USceneComponent* TrackingReferenceComponent, FVector& TRLocation, FRotator& TRRotation)
 {
 	if (!TrackingReferenceComponent)
 	{
@@ -118,9 +120,9 @@ bool UOculusMRFunctionLibrary::GetTrackingReferenceLocationAndRotationInWorldSpa
 	}
 }
 
-UOculusMR_Settings* UOculusMRFunctionLibrary::GetOculusMRSettings()
+UDEPRECATED_UOculusMR_Settings* UDEPRECATED_UOculusMRFunctionLibrary::GetOculusMRSettings()
 {
-	UOculusMR_Settings* Settings = nullptr;
+	UDEPRECATED_UOculusMR_Settings* Settings = nullptr;
 	if (FOculusMRModule::IsAvailable())
 	{
 		Settings = FOculusMRModule::Get().GetMRSettings();
@@ -128,7 +130,7 @@ UOculusMR_Settings* UOculusMRFunctionLibrary::GetOculusMRSettings()
 	return Settings;
 }
 
-USceneComponent* UOculusMRFunctionLibrary::GetTrackingReferenceComponent()
+USceneComponent* UDEPRECATED_UOculusMRFunctionLibrary::GetTrackingReferenceComponent()
 {
 	USceneComponent * TrackingRef = nullptr;
 	if (FOculusMRModule::IsAvailable())
@@ -138,7 +140,7 @@ USceneComponent* UOculusMRFunctionLibrary::GetTrackingReferenceComponent()
 	return TrackingRef;
 }
 
-bool UOculusMRFunctionLibrary::SetTrackingReferenceComponent(USceneComponent* Component)
+bool UDEPRECATED_UOculusMRFunctionLibrary::SetTrackingReferenceComponent(USceneComponent* Component)
 {
 	if (FOculusMRModule::IsAvailable())
 	{
@@ -148,7 +150,7 @@ bool UOculusMRFunctionLibrary::SetTrackingReferenceComponent(USceneComponent* Co
 	return false;
 }
 
-float UOculusMRFunctionLibrary::GetMrcScalingFactor()
+float UDEPRECATED_UOculusMRFunctionLibrary::GetMrcScalingFactor()
 {
 	if (FOculusMRModule::IsAvailable())
 	{
@@ -157,7 +159,7 @@ float UOculusMRFunctionLibrary::GetMrcScalingFactor()
 	return 0.0;
 }
 
-bool UOculusMRFunctionLibrary::SetMrcScalingFactor(float ScalingFactor)
+bool UDEPRECATED_UOculusMRFunctionLibrary::SetMrcScalingFactor(float ScalingFactor)
 {
 	if (FOculusMRModule::IsAvailable() && ScalingFactor > 0.0f)
 	{
@@ -167,12 +169,14 @@ bool UOculusMRFunctionLibrary::SetMrcScalingFactor(float ScalingFactor)
 	return false;
 }
 
-bool UOculusMRFunctionLibrary::IsMrcEnabled()
+bool UDEPRECATED_UOculusMRFunctionLibrary::IsMrcEnabled()
 {
 	return FOculusMRModule::IsAvailable() && FOculusMRModule::Get().IsInitialized();
 }
 
-bool UOculusMRFunctionLibrary::IsMrcActive()
+bool UDEPRECATED_UOculusMRFunctionLibrary::IsMrcActive()
 {
 	return FOculusMRModule::IsAvailable() && FOculusMRModule::Get().IsActive();
-}
\ No newline at end of file
+}
+
+PRAGMA_ENABLE_DEPRECATION_WARNINGS
diff --git a/Engine/Plugins/Runtime/Oculus/OculusVR/Source/OculusMR/Private/OculusMRModule.cpp b/Engine/Plugins/Runtime/Oculus/OculusVR/Source/OculusMR/Private/OculusMRModule.cpp
--- a/Engine/Plugins/Runtime/Oculus/OculusVR/Source/OculusMR/Private/OculusMRModule.cpp
+++ b/Engine/Plugins/Runtime/Oculus/OculusVR/Source/OculusMR/Private/OculusMRModule.cpp
@@ -29,6 +29,8 @@
 
 #define LOCTEXT_NAMESPACE "OculusMR"
 
+PRAGMA_DISABLE_DEPRECATION_WARNINGS
+
 namespace
 {
 	ovrpCameraDevice ConvertCameraDevice(EOculusMR_CameraDeviceEnum device)
@@ -44,7 +46,7 @@ namespace
 
 FOculusMRModule::FOculusMRModule()
 	: bInitialized(false)
-	, MRSettings(nullptr)
+	, MRSettings_DEPRECATED(nullptr)
 	, MRState(nullptr)
 	, MRActor(nullptr)
 	, CurrentWorld(nullptr)
@@ -118,7 +120,7 @@ void FOculusMRModule::ShutdownModule()
 			FEditorDelegates::PrePIEEnded.Remove(PieEndedEventBinding);
 #else
 			// Stop casting and close camera with module if it's the game
-			MRSettings->SetIsCasting(false);
+			MRSettings_DEPRECATED->SetIsCasting(false);
 #endif
 		}
 #if PLATFORM_ANDROID
@@ -130,9 +132,9 @@ void FOculusMRModule::ShutdownModule()
 #endif 
 		FOculusHMDModule::GetPluginWrapper().ShutdownMixedReality();
 
-		if (MRSettings->IsRooted())
+		if (MRSettings_DEPRECATED->IsRooted())
 		{
-			MRSettings->RemoveFromRoot();
+			MRSettings_DEPRECATED->RemoveFromRoot();
 		}
 		if (MRState->IsRooted())
 		{
@@ -156,16 +158,16 @@ void FOculusMRModule::ShutdownModule()
 
 bool FOculusMRModule::IsActive()
 {
-	bool bReturn = bInitialized && MRSettings && MRSettings->GetIsCasting();
+	bool bReturn = bInitialized && MRSettings_DEPRECATED && MRSettings_DEPRECATED->GetIsCasting();
 #if PLATFORM_ANDROID
 	bReturn = bReturn && bActivated;
 #endif
 	return bReturn;
 }
 
-UOculusMR_Settings* FOculusMRModule::GetMRSettings()
+UDEPRECATED_UOculusMR_Settings* FOculusMRModule::GetMRSettings()
 {
-	return MRSettings;
+	return MRSettings_DEPRECATED;
 }
 
 UOculusMR_State* FOculusMRModule::GetMRState()
@@ -210,14 +212,14 @@ void FOculusMRModule::InitMixedRealityCapture()
 {
 	bInitialized = true;
 
-	MRSettings = NewObject<UOculusMR_Settings>((UObject*)GetTransientPackage(), FName("OculusMR_Settings"), RF_MarkAsRootSet);
+	MRSettings_DEPRECATED = NewObject<UDEPRECATED_UOculusMR_Settings>((UObject*)GetTransientPackage(), FName("OculusMR_Settings"), RF_MarkAsRootSet);
 	MRState = NewObject<UOculusMR_State>((UObject*)GetTransientPackage(), FName("OculusMR_State"), RF_MarkAsRootSet);
 
 	// Always bind the event handlers in case devs call them without MRC on
-	MRSettings->TrackedCameraIndexChangeDelegate.BindRaw(this, &FOculusMRModule::OnTrackedCameraIndexChanged);
-	MRSettings->CompositionMethodChangeDelegate.BindRaw(this, &FOculusMRModule::OnCompositionMethodChanged);
-	MRSettings->CapturingCameraChangeDelegate.BindRaw(this, &FOculusMRModule::OnCapturingCameraChanged);
-	MRSettings->IsCastingChangeDelegate.BindRaw(this, &FOculusMRModule::OnIsCastingChanged);
+	MRSettings_DEPRECATED->TrackedCameraIndexChangeDelegate.BindRaw(this, &FOculusMRModule::OnTrackedCameraIndexChanged);
+	MRSettings_DEPRECATED->CompositionMethodChangeDelegate.BindRaw(this, &FOculusMRModule::OnCompositionMethodChanged);
+	MRSettings_DEPRECATED->CapturingCameraChangeDelegate.BindRaw(this, &FOculusMRModule::OnCapturingCameraChanged);
+	MRSettings_DEPRECATED->IsCastingChangeDelegate.BindRaw(this, &FOculusMRModule::OnIsCastingChanged);
 
 	ResetSettingsAndState();
 
@@ -232,7 +234,7 @@ void FOculusMRModule::InitMixedRealityCapture()
 	PieEndedEventBinding = FEditorDelegates::PrePIEEnded.AddRaw(this, &FOculusMRModule::OnPieEnded);
 #else // WITH_EDITOR
 	// Start casting and open camera with the module if it's the game
-	MRSettings->SetIsCasting(true);
+	MRSettings_DEPRECATED->SetIsCasting(true);
 #endif // WITH_EDITOR
 }
 
@@ -240,7 +242,7 @@ void FOculusMRModule::SetupExternalCamera()
 {
 	using namespace OculusHMD;
 
-	if (!MRSettings->GetIsCasting())
+	if (!MRSettings_DEPRECATED->GetIsCasting())
 	{
 		return;
 	}
@@ -248,7 +250,7 @@ void FOculusMRModule::SetupExternalCamera()
 	// Always request the MRC actor to handle a camera state change on its end
 	MRState->ChangeCameraStateRequested = true;
 
-	if (MRSettings->CompositionMethod == EOculusMR_CompositionMethod::ExternalComposition)
+	if (MRSettings_DEPRECATED->CompositionMethod == EOculusMR_CompositionMethod::ExternalComposition)
 	{
 		// Close the camera device for external composition since we don't need the actual camera feed
 		if (MRState->CurrentCapturingCamera != ovrpCameraDevice_None)
@@ -257,17 +259,17 @@ void FOculusMRModule::SetupExternalCamera()
 		}
 	}
 #if PLATFORM_WINDOWS
-	else if (MRSettings->CompositionMethod == EOculusMR_CompositionMethod::DirectComposition)
+	else if (MRSettings_DEPRECATED->CompositionMethod == EOculusMR_CompositionMethod::DirectComposition)
 	{
 		ovrpBool available = ovrpBool_False;
-		if (MRSettings->CapturingCamera == EOculusMR_CameraDeviceEnum::CD_None)
+		if (MRSettings_DEPRECATED->CapturingCamera == EOculusMR_CameraDeviceEnum::CD_None)
 		{
 			MRState->CurrentCapturingCamera = ovrpCameraDevice_None;
 			UE_LOG(LogMR, Error, TEXT("CapturingCamera is set to CD_None which is invalid. Please pick a valid camera for CapturingCamera. If you are not sure, try to set it to CD_WebCamera0 and use the first connected USB web camera"));
 			return;
 		}
 
-		MRState->CurrentCapturingCamera = ConvertCameraDevice(MRSettings->CapturingCamera);
+		MRState->CurrentCapturingCamera = ConvertCameraDevice(MRSettings_DEPRECATED->CapturingCamera);
 		if (OVRP_FAILURE(FOculusHMDModule::GetPluginWrapper().IsCameraDeviceAvailable2(MRState->CurrentCapturingCamera, &available)) || !available)
 		{
 			MRState->CurrentCapturingCamera = ovrpCameraDevice_None;
@@ -312,7 +314,7 @@ void FOculusMRModule::CloseExternalCamera()
 void FOculusMRModule::SetupInGameCapture()
 {
 	// Don't do anything if we don't have a UWorld or if we are not casting
-	if (CurrentWorld == nullptr || !MRSettings->GetIsCasting())
+	if (CurrentWorld == nullptr || !MRSettings_DEPRECATED->GetIsCasting())
 	{
 		return;
 	}
@@ -332,7 +334,7 @@ void FOculusMRModule::SetupInGameCapture()
 
 	// Spawn an MRC camera actor if one wasn't already there
 	MRActor = CurrentWorld->SpawnActorDeferred<AOculusMR_CastingCameraActor>(AOculusMR_CastingCameraActor::StaticClass(), FTransform::Identity);
-	MRActor->InitializeStates(MRSettings, MRState);
+	MRActor->InitializeStates(MRSettings_DEPRECATED, MRState);
 	UGameplayStatics::FinishSpawningActor(MRActor, FTransform::Identity);
 }
 
@@ -358,19 +360,19 @@ void FOculusMRModule::ResetSettingsAndState()
 	// Reset MR Settings
 	const bool bAutoOpenInExternalComposition = FParse::Param(FCommandLine::Get(), TEXT("externalcomposition"));
 	const bool bAutoOpenInDirectComposition = FParse::Param(FCommandLine::Get(), TEXT("directcomposition"));
-	MRSettings->BindToTrackedCameraIndexIfAvailable(0);
-	MRSettings->LoadFromIni();
+	MRSettings_DEPRECATED->BindToTrackedCameraIndexIfAvailable(0);
+	MRSettings_DEPRECATED->LoadFromIni();
 
 	// Save right after load to write defaults to the config if they weren't already there
-	MRSettings->SaveToIni();
+	MRSettings_DEPRECATED->SaveToIni();
 
 	if (bAutoOpenInExternalComposition)
 	{
-		MRSettings->CompositionMethod = EOculusMR_CompositionMethod::ExternalComposition;
+		MRSettings_DEPRECATED->CompositionMethod = EOculusMR_CompositionMethod::ExternalComposition;
 	}
 	else if (bAutoOpenInDirectComposition)
 	{
-		MRSettings->CompositionMethod = EOculusMR_CompositionMethod::DirectComposition;
+		MRSettings_DEPRECATED->CompositionMethod = EOculusMR_CompositionMethod::DirectComposition;
 	}
 }
 
@@ -599,7 +601,7 @@ void FOculusMRModule::OnPieBegin(bool bIsSimulating)
 		ResetSettingsAndState();
 
 		// Always start casting with PIE (since this can only be reached if the command line param is on)
-		MRSettings->SetIsCasting(true);
+		MRSettings_DEPRECATED->SetIsCasting(true);
 	}
 }
 
@@ -619,12 +621,14 @@ void FOculusMRModule::OnPieEnded(bool bIsSimulating)
 	if (!bIsSimulating && PieWorld)
 	{
 		// Stop casting when PIE ends
-		MRSettings->SetIsCasting(false);
+		MRSettings_DEPRECATED->SetIsCasting(false);
 		OnWorldDestroyed(PieWorld);
 	}
 }
 #endif // WITH_EDITOR
 
+PRAGMA_ENABLE_DEPRECATION_WARNINGS
+
 IMPLEMENT_MODULE( FOculusMRModule, OculusMR )
 
 #undef LOCTEXT_NAMESPACE
diff --git a/Engine/Plugins/Runtime/Oculus/OculusVR/Source/OculusMR/Private/OculusMRModule.h b/Engine/Plugins/Runtime/Oculus/OculusVR/Source/OculusMR/Private/OculusMRModule.h
--- a/Engine/Plugins/Runtime/Oculus/OculusVR/Source/OculusMR/Private/OculusMRModule.h
+++ b/Engine/Plugins/Runtime/Oculus/OculusVR/Source/OculusMR/Private/OculusMRModule.h
@@ -10,7 +10,7 @@ enum class EOculusMR_CompositionMethod : uint8;
 enum class EOculusMR_CameraDeviceEnum : uint8;
 enum class EOculusMR_DepthQuality : uint8;
 
-class UOculusMR_Settings;
+class UDEPRECATED_UOculusMR_Settings;
 class AOculusMR_CastingCameraActor;
 class UOculusMR_State;
 
@@ -18,6 +18,8 @@ class UOculusMR_State;
 // FOculusInputModule
 //-------------------------------------------------------------------------------------------------
 
+PRAGMA_DISABLE_DEPRECATION_WARNINGS
+
 class FOculusMRModule : public IOculusMRModule
 {
 public:
@@ -36,12 +38,12 @@ public:
 	bool IsInitialized() { return bInitialized; }
 
 	bool IsActive();
-	UOculusMR_Settings* GetMRSettings();
+	UDEPRECATED_UOculusMR_Settings* GetMRSettings();
 	UOculusMR_State* GetMRState();
 
 private:
 	bool bInitialized;
-	UOculusMR_Settings* MRSettings;
+	UDEPRECATED_UOculusMR_Settings* MRSettings_DEPRECATED;
 	UOculusMR_State* MRState;
 	AOculusMR_CastingCameraActor* MRActor;
 	UWorld* CurrentWorld;
@@ -97,5 +99,6 @@ private:
 #endif
 };
 
+PRAGMA_ENABLE_DEPRECATION_WARNINGS
 
 #undef LOCTEXT_NAMESPACE
diff --git a/Engine/Plugins/Runtime/Oculus/OculusVR/Source/OculusMR/Private/OculusMR_CastingCameraActor.cpp b/Engine/Plugins/Runtime/Oculus/OculusVR/Source/OculusMR/Private/OculusMR_CastingCameraActor.cpp
--- a/Engine/Plugins/Runtime/Oculus/OculusVR/Source/OculusMR/Private/OculusMR_CastingCameraActor.cpp
+++ b/Engine/Plugins/Runtime/Oculus/OculusVR/Source/OculusMR/Private/OculusMR_CastingCameraActor.cpp
@@ -39,6 +39,8 @@ static TAutoConsoleVariable<float> CChromaKeySmoothRangeVar(TEXT("oculus.mr.Chro
 static TAutoConsoleVariable<float> CChromaKeySpillRangeVar(TEXT("oculus.mr.ChromaKeySpillRange"), 0.04f, TEXT("Chroma Key Spill Range"));
 static TAutoConsoleVariable<float> CCastingLantencyVar(TEXT("oculus.mr.CastingLantency"), 0, TEXT("Casting Latency"));
 
+PRAGMA_DISABLE_DEPRECATION_WARNINGS
+
 namespace
 {
 	bool GetCameraTrackedObjectPoseInTrackingSpace(OculusHMD::FOculusHMD* OculusHMD, const FTrackedCamera& TrackedCamera, OculusHMD::FPose& CameraTrackedObjectPose)
@@ -278,11 +280,11 @@ void AOculusMR_CastingCameraActor::Tick(float DeltaTime)
 
 	if (COverrideMixedRealityParametersVar.GetValueOnAnyThread() > 0)
 	{
-		MRSettings->ChromaKeyColor = FColor(CChromaKeyColorRVar.GetValueOnAnyThread(), CChromaKeyColorGVar.GetValueOnAnyThread(), CChromaKeyColorBVar.GetValueOnAnyThread());
-		MRSettings->ChromaKeySimilarity = CChromaKeySimilarityVar.GetValueOnAnyThread();
-		MRSettings->ChromaKeySmoothRange = CChromaKeySmoothRangeVar.GetValueOnAnyThread();
-		MRSettings->ChromaKeySpillRange = CChromaKeySpillRangeVar.GetValueOnAnyThread();
-		MRSettings->CastingLatency = CCastingLantencyVar.GetValueOnAnyThread();
+		MRSettings_DEPRECATED->ChromaKeyColor = FColor(CChromaKeyColorRVar.GetValueOnAnyThread(), CChromaKeyColorGVar.GetValueOnAnyThread(), CChromaKeyColorBVar.GetValueOnAnyThread());
+		MRSettings_DEPRECATED->ChromaKeySimilarity = CChromaKeySimilarityVar.GetValueOnAnyThread();
+		MRSettings_DEPRECATED->ChromaKeySmoothRange = CChromaKeySmoothRangeVar.GetValueOnAnyThread();
+		MRSettings_DEPRECATED->ChromaKeySpillRange = CChromaKeySpillRangeVar.GetValueOnAnyThread();
+		MRSettings_DEPRECATED->CastingLatency = CCastingLantencyVar.GetValueOnAnyThread();
 	}
 
 	// Reset capturing components if the composition method changes
@@ -295,16 +297,16 @@ void AOculusMR_CastingCameraActor::Tick(float DeltaTime)
 	}
 
 #if PLATFORM_WINDOWS
-	if (MRSettings->GetCompositionMethod() == EOculusMR_CompositionMethod::ExternalComposition)
+	if (MRSettings_DEPRECATED->GetCompositionMethod() == EOculusMR_CompositionMethod::ExternalComposition)
 #endif
 	{
-		if (ForegroundLayerBackgroundColor != MRSettings->BackdropColor)
+		if (ForegroundLayerBackgroundColor != MRSettings_DEPRECATED->BackdropColor)
 		{
-			ForegroundLayerBackgroundColor = MRSettings->BackdropColor;
+			ForegroundLayerBackgroundColor = MRSettings_DEPRECATED->BackdropColor;
 			SetBackdropMaterialColor();
 		}
 		// Enable external composition post process based on setting
-		bool bPostProcess = MRSettings->ExternalCompositionPostProcessEffects != EOculusMR_PostProcessEffects::PPE_Off;
+		bool bPostProcess = MRSettings_DEPRECATED->ExternalCompositionPostProcessEffects != EOculusMR_PostProcessEffects::PPE_Off;
 		if (COverrideMixedRealityParametersVar.GetValueOnAnyThread() > 0)
 		{
 			bPostProcess = CEnableExternalCompositionPostProcess.GetValueOnAnyThread() > 0;
@@ -316,16 +318,16 @@ void AOculusMR_CastingCameraActor::Tick(float DeltaTime)
 		}
 	}
 #if PLATFORM_WINDOWS
-	else if (MRSettings->GetCompositionMethod() == EOculusMR_CompositionMethod::DirectComposition)
+	else if (MRSettings_DEPRECATED->GetCompositionMethod() == EOculusMR_CompositionMethod::DirectComposition)
 	{
 		SetupCameraFrameMaterialInstance();
 
 		if (CameraFrameMaterialInstance)
 		{
-			CameraFrameMaterialInstance->SetVectorParameterValue(FName(TEXT("ChromaKeyColor")), FLinearColor(MRSettings->ChromaKeyColor));
-			CameraFrameMaterialInstance->SetScalarParameterValue(FName(TEXT("ChromaKeySimilarity")), MRSettings->ChromaKeySimilarity);
-			CameraFrameMaterialInstance->SetScalarParameterValue(FName(TEXT("ChromaKeySmoothRange")), MRSettings->ChromaKeySmoothRange);
-			CameraFrameMaterialInstance->SetScalarParameterValue(FName(TEXT("ChromaKeySpillRange")), MRSettings->ChromaKeySpillRange);
+			CameraFrameMaterialInstance->SetVectorParameterValue(FName(TEXT("ChromaKeyColor")), FLinearColor(MRSettings_DEPRECATED->ChromaKeyColor));
+			CameraFrameMaterialInstance->SetScalarParameterValue(FName(TEXT("ChromaKeySimilarity")), MRSettings_DEPRECATED->ChromaKeySimilarity);
+			CameraFrameMaterialInstance->SetScalarParameterValue(FName(TEXT("ChromaKeySmoothRange")), MRSettings_DEPRECATED->ChromaKeySmoothRange);
+			CameraFrameMaterialInstance->SetScalarParameterValue(FName(TEXT("ChromaKeySpillRange")), MRSettings_DEPRECATED->ChromaKeySpillRange);
 		}
 	}
 #endif
@@ -355,7 +357,7 @@ void AOculusMR_CastingCameraActor::Tick(float DeltaTime)
 #if PLATFORM_WINDOWS
 	RepositionPlaneMesh();
 
-	double HandPoseStateLatencyToSet = (double)MRSettings->HandPoseStateLatency;
+	double HandPoseStateLatencyToSet = (double)MRSettings_DEPRECATED->HandPoseStateLatency;
 	ovrpResult result = FOculusHMDModule::GetPluginWrapper().SetHandNodePoseStateLatency(HandPoseStateLatencyToSet);
 	if (OVRP_FAILURE(result))
 	{
@@ -534,14 +536,14 @@ void AOculusMR_CastingCameraActor::Execute_BindToTrackedCameraIndexIfAvailable()
 	}
 
 	FTrackedCamera TempTrackedCamera;
-	if (MRSettings->GetBindToTrackedCameraIndex() >= 0)
+	if (MRSettings_DEPRECATED->GetBindToTrackedCameraIndex() >= 0)
 	{
 		TArray<FTrackedCamera> TrackedCameras;
-		UOculusMRFunctionLibrary::GetAllTrackedCamera(TrackedCameras);
+		UDEPRECATED_UOculusMRFunctionLibrary::GetAllTrackedCamera(TrackedCameras);
 		int i;
 		for (i = 0; i < TrackedCameras.Num(); ++i)
 		{
-			if (TrackedCameras[i].Index == MRSettings->GetBindToTrackedCameraIndex())
+			if (TrackedCameras[i].Index == MRSettings_DEPRECATED->GetBindToTrackedCameraIndex())
 			{
 				TempTrackedCamera = TrackedCameras[i];
 				break;
@@ -549,12 +551,12 @@ void AOculusMR_CastingCameraActor::Execute_BindToTrackedCameraIndexIfAvailable()
 		}
 		if (i == TrackedCameras.Num())
 		{
-			UE_LOG(LogMR, Warning, TEXT("Unable to find TrackedCamera at index %d, use TempTrackedCamera"), MRSettings->GetBindToTrackedCameraIndex());
+			UE_LOG(LogMR, Warning, TEXT("Unable to find TrackedCamera at index %d, use TempTrackedCamera"), MRSettings_DEPRECATED->GetBindToTrackedCameraIndex());
 		}
 	}
 	else
 	{
-		UE_LOG(LogMR, Warning, TEXT("BindToTrackedCameraIndex == %d, use TempTrackedCamera"), MRSettings->GetBindToTrackedCameraIndex());
+		UE_LOG(LogMR, Warning, TEXT("BindToTrackedCameraIndex == %d, use TempTrackedCamera"), MRSettings_DEPRECATED->GetBindToTrackedCameraIndex());
 	}
 
 	MRState->TrackedCamera = TempTrackedCamera;
@@ -601,7 +603,7 @@ void AOculusMR_CastingCameraActor::SetTrackedCameraInitialPoseWithPlayerTransfor
 	FQuat TROrientation;
 	FVector TRLocation;
 	FRotator TRRotation;
-	if (!UOculusMRFunctionLibrary::GetTrackingReferenceLocationAndRotationInWorldSpace(MRState->TrackingReferenceComponent, TRLocation, TRRotation))
+	if (!UDEPRECATED_UOculusMRFunctionLibrary::GetTrackingReferenceLocationAndRotationInWorldSpace(MRState->TrackingReferenceComponent, TRLocation, TRRotation))
 	{
 		UE_LOG(LogMR, Warning, TEXT("Could not get player position"));
 		return;
@@ -646,7 +648,7 @@ void AOculusMR_CastingCameraActor::SetTrackedCameraUserPoseWithCameraTransform()
 	FQuat TROrientation;
 	FVector TRLocation;
 	FRotator TRRotation;
-	if (!UOculusMRFunctionLibrary::GetTrackingReferenceLocationAndRotationInWorldSpace(MRState->TrackingReferenceComponent, TRLocation, TRRotation))
+	if (!UDEPRECATED_UOculusMRFunctionLibrary::GetTrackingReferenceLocationAndRotationInWorldSpace(MRState->TrackingReferenceComponent, TRLocation, TRRotation))
 	{
 		UE_LOG(LogMR, Warning, TEXT("Could not get player position"));
 		return;
@@ -694,11 +696,11 @@ void AOculusMR_CastingCameraActor::UpdateTrackedCameraPosition()
 	CameraPose.Position = CameraPose.Position * MRState->ScalingFactor;
 
 	float Distance = 0.0f;
-	if (MRSettings->ClippingReference == EOculusMR_ClippingReference::CR_TrackingReference)
+	if (MRSettings_DEPRECATED->ClippingReference == EOculusMR_ClippingReference::CR_TrackingReference)
 	{
 		Distance = -FVector::DotProduct(CameraPose.Orientation.GetForwardVector().GetSafeNormal2D(), CameraPose.Position);
 	}
-	else if (MRSettings->ClippingReference == EOculusMR_ClippingReference::CR_Head)
+	else if (MRSettings_DEPRECATED->ClippingReference == EOculusMR_ClippingReference::CR_Head)
 	{
 		FQuat HeadOrientation;
 		FVector HeadPosition;
@@ -720,7 +722,7 @@ void AOculusMR_CastingCameraActor::UpdateTrackedCameraPosition()
 	FQuat TROrientation;
 	FVector TRLocation;
 	FRotator TRRotation;
-	if (!UOculusMRFunctionLibrary::GetTrackingReferenceLocationAndRotationInWorldSpace(MRState->TrackingReferenceComponent, TRLocation, TRRotation))
+	if (!UDEPRECATED_UOculusMRFunctionLibrary::GetTrackingReferenceLocationAndRotationInWorldSpace(MRState->TrackingReferenceComponent, TRLocation, TRRotation))
 	{
 		UE_LOG(LogMR, Warning, TEXT("Could not get player position"));
 		return;
@@ -739,9 +741,9 @@ void AOculusMR_CastingCameraActor::UpdateTrackedCameraPosition()
 	}
 }
 
-void AOculusMR_CastingCameraActor::InitializeStates(UOculusMR_Settings* MRSettingsIn, UOculusMR_State* MRStateIn)
+void AOculusMR_CastingCameraActor::InitializeStates(UDEPRECATED_UOculusMR_Settings* MRSettingsIn, UOculusMR_State* MRStateIn)
 {
-	MRSettings = MRSettingsIn;
+	MRSettings_DEPRECATED = MRSettingsIn;
 	MRState = MRStateIn;
 }
 
@@ -759,7 +761,7 @@ void AOculusMR_CastingCameraActor::SetupTrackedCamera()
 
 #if PLATFORM_WINDOWS
 	// Set the plane mesh to the camera stream in direct composition or static background for external composition
-	if (MRSettings->GetCompositionMethod() == EOculusMR_CompositionMethod::DirectComposition)
+	if (MRSettings_DEPRECATED->GetCompositionMethod() == EOculusMR_CompositionMethod::DirectComposition)
 	{
 		ovrpBool cameraOpen;
 		if (OVRP_SUCCESS(FOculusHMDModule::GetPluginWrapper().HasCameraDeviceOpened2(MRState->CurrentCapturingCamera, &cameraOpen)) && cameraOpen)
@@ -778,7 +780,7 @@ void AOculusMR_CastingCameraActor::SetupTrackedCamera()
 
 		SetupCameraFrameMaterialInstance();
 	}
-	else if (MRSettings->GetCompositionMethod() == EOculusMR_CompositionMethod::ExternalComposition)
+	else if (MRSettings_DEPRECATED->GetCompositionMethod() == EOculusMR_CompositionMethod::ExternalComposition)
 	{
 		SetupBackdropMaterialInstance();
 	}
@@ -829,8 +831,8 @@ void AOculusMR_CastingCameraActor::RepositionPlaneMesh()
 	FVector PlaneCenter = FVector::ForwardVector * ForegroundMaxDistance;
 	FVector PlaneUp = FVector::UpVector;
 	FVector PlaneNormal = -FVector::ForwardVector;
-	int ViewWidth = MRSettings->bUseTrackedCameraResolution ? MRState->TrackedCamera.SizeX : MRSettings->WidthPerView;
-	int ViewHeight = MRSettings->bUseTrackedCameraResolution ? MRState->TrackedCamera.SizeY : MRSettings->HeightPerView;
+	int ViewWidth = MRSettings_DEPRECATED->bUseTrackedCameraResolution ? MRState->TrackedCamera.SizeX : MRSettings_DEPRECATED->WidthPerView;
+	int ViewHeight = MRSettings_DEPRECATED->bUseTrackedCameraResolution ? MRState->TrackedCamera.SizeY : MRSettings_DEPRECATED->HeightPerView;
 	float Width = ForegroundMaxDistance * FMath::Tan(FMath::DegreesToRadians(GetCaptureComponent2D()->FOVAngle) * 0.5f) * 2.0f;
 	float Height = Width * ViewHeight / ViewWidth;
 	FVector2D PlaneSize = FVector2D(Width, Height);
@@ -885,8 +887,8 @@ void BuildProjectionMatrix(float YMultiplier, float FOV, float FarClipPlane, FMa
 
 void AOculusMR_CastingCameraActor::UpdateRenderTargetSize()
 {
-	int ViewWidth = MRSettings->bUseTrackedCameraResolution ? MRState->TrackedCamera.SizeX : MRSettings->WidthPerView;
-	int ViewHeight = MRSettings->bUseTrackedCameraResolution ? MRState->TrackedCamera.SizeY : MRSettings->HeightPerView;
+	int ViewWidth = MRSettings_DEPRECATED->bUseTrackedCameraResolution ? MRState->TrackedCamera.SizeX : MRSettings_DEPRECATED->WidthPerView;
+	int ViewHeight = MRSettings_DEPRECATED->bUseTrackedCameraResolution ? MRState->TrackedCamera.SizeY : MRSettings_DEPRECATED->HeightPerView;
 
 #if PLATFORM_WINDOWS
 	BackgroundRenderTargets[0]->ResizeTarget(ViewWidth, ViewHeight);
@@ -946,7 +948,7 @@ void AOculusMR_CastingCameraActor::SetupMRCScreen()
 		GetCaptureComponent2D()->TextureTarget = BackgroundRenderTargets[0];
 
 #if PLATFORM_WINDOWS
-		if (MRSettings->GetCompositionMethod() == EOculusMR_CompositionMethod::ExternalComposition)
+		if (MRSettings_DEPRECATED->GetCompositionMethod() == EOculusMR_CompositionMethod::ExternalComposition)
 #endif
 		{
 			ForegroundCaptureActor = GetWorld()->SpawnActor<ASceneCapture2D>();
@@ -975,7 +977,7 @@ void AOculusMR_CastingCameraActor::SetupMRCScreen()
 			ForegroundCaptureActor->AttachToActor(this, FAttachmentTransformRules(EAttachmentRule::SnapToTarget, true));
 		}
 #if PLATFORM_WINDOWS
-		else if (MRSettings->GetCompositionMethod() == EOculusMR_CompositionMethod::DirectComposition)
+		else if (MRSettings_DEPRECATED->GetCompositionMethod() == EOculusMR_CompositionMethod::DirectComposition)
 		{
 			SpecScreen->SetMRBackground(BackgroundRenderTargets[0]);
 			SpecScreen->SetMRSpectatorScreenMode(OculusHMD::EMRSpectatorScreenMode::DirectComposition);
@@ -1027,4 +1029,6 @@ void AOculusMR_CastingCameraActor::CloseTrackedCamera()
 	CameraFrameMaterialInstance = NULL;
 }
 
+PRAGMA_ENABLE_DEPRECATION_WARNINGS
+
 #undef LOCTEXT_NAMESPACE
diff --git a/Engine/Plugins/Runtime/Oculus/OculusVR/Source/OculusMR/Private/OculusMR_CastingCameraActor.h b/Engine/Plugins/Runtime/Oculus/OculusVR/Source/OculusMR/Private/OculusMR_CastingCameraActor.h
--- a/Engine/Plugins/Runtime/Oculus/OculusVR/Source/OculusMR/Private/OculusMR_CastingCameraActor.h
+++ b/Engine/Plugins/Runtime/Oculus/OculusVR/Source/OculusMR/Private/OculusMR_CastingCameraActor.h
@@ -17,7 +17,7 @@ class UOculusMR_PlaneMeshComponent;
 class UMaterial;
 class AOculusMR_BoundaryActor;
 class UTextureRenderTarget2D;
-class UOculusMR_Settings;
+class UDEPRECATED_UOculusMR_Settings;
 class UOculusMR_State;
 class UMaterialInstanceDynamic;
 
@@ -33,7 +33,7 @@ public:
 	AOculusMR_CastingCameraActor(const FObjectInitializer& ObjectInitializer);
 
 	/** Initialize the MRC settings and states */
-	void InitializeStates(UOculusMR_Settings* MRSettingsIn, UOculusMR_State* MRStateIn);
+	void InitializeStates(UDEPRECATED_UOculusMR_Settings* MRSettingsIn, UOculusMR_State* MRStateIn);
 
 	virtual void BeginPlay() override;
 	virtual void EndPlay(EEndPlayReason::Type Reason) override;
@@ -129,8 +129,8 @@ private:
 	UPROPERTY()
 	TArray<double> PoseTimes;
 
-	UPROPERTY()
-	TObjectPtr<UOculusMR_Settings> MRSettings;
+	UPROPERTY(meta = (DeprecatedProperty, DeprecationMessage = "The OculusVR plugin is deprecated."))
+	TObjectPtr<UDEPRECATED_UOculusMR_Settings> MRSettings_DEPRECATED;
 
 	UPROPERTY()
 	TObjectPtr<UOculusMR_State> MRState;
diff --git a/Engine/Plugins/Runtime/Oculus/OculusVR/Source/OculusMR/Private/OculusMR_Settings.cpp b/Engine/Plugins/Runtime/Oculus/OculusVR/Source/OculusMR/Private/OculusMR_Settings.cpp
--- a/Engine/Plugins/Runtime/Oculus/OculusVR/Source/OculusMR/Private/OculusMR_Settings.cpp
+++ b/Engine/Plugins/Runtime/Oculus/OculusVR/Source/OculusMR/Private/OculusMR_Settings.cpp
@@ -5,7 +5,9 @@
 #include "OculusHMD.h"
 #include "Engine/Engine.h"
 
-UOculusMR_Settings::UOculusMR_Settings(const FObjectInitializer& ObjectInitializer)
+PRAGMA_DISABLE_DEPRECATION_WARNINGS
+
+UDEPRECATED_UOculusMR_Settings::UDEPRECATED_UOculusMR_Settings(const FObjectInitializer& ObjectInitializer)
 	: Super(ObjectInitializer)
 	, ClippingReference(EOculusMR_ClippingReference::CR_Head)
 	, bUseTrackedCameraResolution(true)
@@ -26,7 +28,7 @@ UOculusMR_Settings::UOculusMR_Settings(const FObjectInitializer& ObjectInitializ
 {
 }
 
-void UOculusMR_Settings::SetCompositionMethod(EOculusMR_CompositionMethod val)
+void UDEPRECATED_UOculusMR_Settings::SetCompositionMethod(EOculusMR_CompositionMethod val)
 {
 	if (CompositionMethod == val)
 	{
@@ -37,7 +39,7 @@ void UOculusMR_Settings::SetCompositionMethod(EOculusMR_CompositionMethod val)
 	CompositionMethodChangeDelegate.Execute(old, val);
 }
 
-void UOculusMR_Settings::SetCapturingCamera(EOculusMR_CameraDeviceEnum val)
+void UDEPRECATED_UOculusMR_Settings::SetCapturingCamera(EOculusMR_CameraDeviceEnum val)
 {
 	if (CapturingCamera == val)
 	{
@@ -48,7 +50,7 @@ void UOculusMR_Settings::SetCapturingCamera(EOculusMR_CameraDeviceEnum val)
 	CapturingCameraChangeDelegate.Execute(old, val);
 }
 
-void UOculusMR_Settings::SetIsCasting(bool val)
+void UDEPRECATED_UOculusMR_Settings::SetIsCasting(bool val)
 {
 	if (bIsCasting == val)
 	{
@@ -59,7 +61,7 @@ void UOculusMR_Settings::SetIsCasting(bool val)
 	IsCastingChangeDelegate.Execute(old, val);
 }
 
-void UOculusMR_Settings::BindToTrackedCameraIndexIfAvailable(int InTrackedCameraIndex)
+void UDEPRECATED_UOculusMR_Settings::BindToTrackedCameraIndexIfAvailable(int InTrackedCameraIndex)
 {
 	if (BindToTrackedCameraIndex == InTrackedCameraIndex)
 	{
@@ -70,7 +72,7 @@ void UOculusMR_Settings::BindToTrackedCameraIndexIfAvailable(int InTrackedCamera
 	TrackedCameraIndexChangeDelegate.Execute(old, InTrackedCameraIndex);
 }
 
-void UOculusMR_Settings::LoadFromIni()
+void UDEPRECATED_UOculusMR_Settings::LoadFromIni()
 {
 	if (!GConfig)
 	{
@@ -152,7 +154,7 @@ void UOculusMR_Settings::LoadFromIni()
 	UE_LOG(LogMR, Log, TEXT("MixedReality settings loaded from Engine.ini"));
 }
 
-void UOculusMR_Settings::SaveToIni() const
+void UDEPRECATED_UOculusMR_Settings::SaveToIni() const
 {
 	if (!GConfig)
 	{
@@ -181,3 +183,5 @@ void UOculusMR_Settings::SaveToIni() const
 
 	UE_LOG(LogMR, Log, TEXT("MixedReality settings saved to Engine.ini"));
 }
+
+PRAGMA_ENABLE_DEPRECATION_WARNINGS
diff --git a/Engine/Plugins/Runtime/Oculus/OculusVR/Source/OculusMR/Public/IOculusMRModule.h b/Engine/Plugins/Runtime/Oculus/OculusVR/Source/OculusMR/Public/IOculusMRModule.h
--- a/Engine/Plugins/Runtime/Oculus/OculusVR/Source/OculusMR/Public/IOculusMRModule.h
+++ b/Engine/Plugins/Runtime/Oculus/OculusVR/Source/OculusMR/Public/IOculusMRModule.h
@@ -9,6 +9,7 @@
  * The public interface to this module.  In most cases, this interface is only public to sibling modules 
  * within this plugin.
  */
+class UE_DEPRECATED(5.1, "OculusVR plugin is deprecated") IOculusMRModule;
 class IOculusMRModule : public IModuleInterface
 {
 
diff --git a/Engine/Plugins/Runtime/Oculus/OculusVR/Source/OculusMR/Public/OculusMRFunctionLibrary.h b/Engine/Plugins/Runtime/Oculus/OculusVR/Source/OculusMR/Public/OculusMRFunctionLibrary.h
--- a/Engine/Plugins/Runtime/Oculus/OculusVR/Source/OculusMR/Public/OculusMRFunctionLibrary.h
+++ b/Engine/Plugins/Runtime/Oculus/OculusVR/Source/OculusMR/Public/OculusMRFunctionLibrary.h
@@ -6,7 +6,7 @@
 #include "OculusMRFunctionLibrary.generated.h"
 
 class USceneComponent;
-class UOculusMR_Settings;
+class UDEPRECATED_UOculusMR_Settings;
 struct FTrackedCamera;
 
 namespace OculusHMD
@@ -14,38 +14,45 @@ namespace OculusHMD
 	class FOculusHMD;
 }
 
-UCLASS()
-class OCULUSMR_API UOculusMRFunctionLibrary : public UBlueprintFunctionLibrary
+UCLASS(deprecated, meta = (DeprecationMessage = "OculusVR plugin is deprecated; please use the built-in OpenXR plugin or OculusXR plugin from the Marketplace."))
+class OCULUSMR_API UDEPRECATED_UOculusMRFunctionLibrary : public UBlueprintFunctionLibrary
 {
 	GENERATED_UCLASS_BODY()
 
 public:
 	// Get the OculusMR settings object
-	UFUNCTION(BlueprintCallable, Category = "OculusLibrary|MR", meta = (DisplayName = "Get Oculus MR Settings"))
-	static UOculusMR_Settings* GetOculusMRSettings();
+	UE_DEPRECATED(5.1, "OculusVR plugin is deprecated; please use the built-in OpenXR plugin or OculusXR plugin from the Marketplace.")
+	UFUNCTION(BlueprintCallable, Category = "OculusLibrary|MR", meta = (DisplayName = "Get Oculus MR Settings", DeprecatedFunction, DeprecationMessage = "The OculusVR plugin is deprecated."))
+	static UDEPRECATED_UOculusMR_Settings* GetOculusMRSettings();
 
 	// Get the component that the OculusMR camera is tracking. When this is null, the camera will track the player pawn.
-	UFUNCTION(BlueprintCallable, Category = "OculusLibrary|MR")
+	UE_DEPRECATED(5.1, "OculusVR plugin is deprecated; please use the built-in OpenXR plugin or OculusXR plugin from the Marketplace.")
+	UFUNCTION(BlueprintCallable, Category = "OculusLibrary|MR", meta = (DeprecatedFunction, DeprecationMessage = "The OculusVR plugin is deprecated."))
 	static USceneComponent* GetTrackingReferenceComponent();
  
 	// Set the component for the OculusMR camera to track. If this is set to null, the camera will track the player pawn.
-	UFUNCTION(BlueprintCallable, Category = "OculusLibrary|MR")
+	UE_DEPRECATED(5.1, "OculusVR plugin is deprecated; please use the built-in OpenXR plugin or OculusXR plugin from the Marketplace.")
+	UFUNCTION(BlueprintCallable, Category = "OculusLibrary|MR", meta = (DeprecatedFunction, DeprecationMessage = "The OculusVR plugin is deprecated."))
 	static bool SetTrackingReferenceComponent(USceneComponent* Component);
 
 	// Get the scaling factor for the MRC configuration. Returns 0 if not available.
-	UFUNCTION(BlueprintCallable, Category = "OculusLibrary|MR", meta = (DisplayName = "Get MRC Scaling Factor"))
+	UE_DEPRECATED(5.1, "OculusVR plugin is deprecated; please use the built-in OpenXR plugin or OculusXR plugin from the Marketplace.")
+	UFUNCTION(BlueprintCallable, Category = "OculusLibrary|MR", meta = (DisplayName = "Get MRC Scaling Factor", DeprecatedFunction, DeprecationMessage = "The OculusVR plugin is deprecated."))
 	static float GetMrcScalingFactor();
 
 	// Set the scaling factor for the MRC configuration. This should be a positive value set to the same scaling as the VR player pawn so that the game capture and camera video are aligned.
-	UFUNCTION(BlueprintCallable, Category = "OculusLibrary|MR", meta = (DisplayName = "Set MRC Scaling Factor"))
+	UE_DEPRECATED(5.1, "OculusVR plugin is deprecated; please use the built-in OpenXR plugin or OculusXR plugin from the Marketplace.")
+	UFUNCTION(BlueprintCallable, Category = "OculusLibrary|MR", meta = (DisplayName = "Set MRC Scaling Factor", DeprecatedFunction, DeprecationMessage = "The OculusVR plugin is deprecated."))
 	static bool SetMrcScalingFactor(float ScalingFactor = 1.0f);
 
 	// Check if MRC is enabled
-	UFUNCTION(BlueprintCallable, Category = "OculusLibrary|MR")
+	UE_DEPRECATED(5.1, "OculusVR plugin is deprecated; please use the built-in OpenXR plugin or OculusXR plugin from the Marketplace.")
+	UFUNCTION(BlueprintCallable, Category = "OculusLibrary|MR", meta = (DeprecatedFunction, DeprecationMessage = "The OculusVR plugin is deprecated."))
 	static bool IsMrcEnabled();
 
 	// Check if MRC is enabled and actively capturing
-	UFUNCTION(BlueprintCallable, Category = "OculusLibrary|MR")
+	UE_DEPRECATED(5.1, "OculusVR plugin is deprecated; please use the built-in OpenXR plugin or OculusXR plugin from the Marketplace.")
+	UFUNCTION(BlueprintCallable, Category = "OculusLibrary|MR", meta = (DeprecatedFunction, DeprecationMessage = "The OculusVR plugin is deprecated."))
 	static bool IsMrcActive();
 
 public:
diff --git a/Engine/Plugins/Runtime/Oculus/OculusVR/Source/OculusMR/Public/OculusMR_Settings.h b/Engine/Plugins/Runtime/Oculus/OculusVR/Source/OculusMR/Public/OculusMR_Settings.h
--- a/Engine/Plugins/Runtime/Oculus/OculusVR/Source/OculusMR/Public/OculusMR_Settings.h
+++ b/Engine/Plugins/Runtime/Oculus/OculusVR/Source/OculusMR/Public/OculusMR_Settings.h
@@ -5,8 +5,10 @@
 
 #include "OculusMR_Settings.generated.h"
 
+PRAGMA_DISABLE_DEPRECATION_WARNINGS
+
 UENUM(BlueprintType)
-enum class EOculusMR_CameraDeviceEnum : uint8
+enum class UE_DEPRECATED(5.1, "OculusVR plugin is deprecated; please use the built-in OpenXR plugin or OculusXR plugin from the Marketplace.") EOculusMR_CameraDeviceEnum : uint8
 {
 	CD_None         UMETA(DisplayName = "None"),
 	CD_WebCamera0   UMETA(DisplayName = "Web Camera 0"),
@@ -14,21 +16,21 @@ enum class EOculusMR_CameraDeviceEnum : uint8
 };
 
 UENUM(BlueprintType)
-enum class EOculusMR_ClippingReference : uint8
+enum class UE_DEPRECATED(5.1, "OculusVR plugin is deprecated; please use the built-in OpenXR plugin or OculusXR plugin from the Marketplace.") EOculusMR_ClippingReference : uint8
 {
 	CR_TrackingReference    UMETA(DisplayName = "Tracking Reference"),
 	CR_Head                 UMETA(DisplayName = "Head"),
 };
 
 UENUM(BlueprintType)
-enum class EOculusMR_PostProcessEffects : uint8
+enum class UE_DEPRECATED(5.1, "OculusVR plugin is deprecated; please use the built-in OpenXR plugin or OculusXR plugin from the Marketplace.") EOculusMR_PostProcessEffects : uint8
 {
 	PPE_Off             UMETA(DisplayName = "Off"),
 	PPE_On				UMETA(DisplayName = "On"),
 };
 
 UENUM(BlueprintType)
-enum class EOculusMR_CompositionMethod : uint8
+enum class UE_DEPRECATED(5.1, "OculusVR plugin is deprecated; please use the built-in OpenXR plugin or OculusXR plugin from the Marketplace.") EOculusMR_CompositionMethod : uint8
 {
 	/* Generate both foreground and background views for compositing with 3rd-party software like OBS. */
 	ExternalComposition		UMETA(DisplayName = "External Composition"),
@@ -36,112 +38,135 @@ enum class EOculusMR_CompositionMethod : uint8
 	DirectComposition		UMETA(DisplayName = "Direct Composition")
 };
 
-UCLASS(ClassGroup = OculusMR, Blueprintable)
-class UOculusMR_Settings : public UObject
+UCLASS(ClassGroup = OculusMR, Blueprintable, deprecated, meta = (DeprecationMessage = "OculusVR plugin is deprecated; please use the built-in OpenXR plugin or OculusXR plugin from the Marketplace."))
+class UDEPRECATED_UOculusMR_Settings : public UObject
 {
 	GENERATED_BODY()
 
 public:
-	UOculusMR_Settings(const FObjectInitializer& ObjectInitializer);
+	UE_DEPRECATED(5.1, "OculusVR plugin is deprecated; please use the built-in OpenXR plugin or OculusXR plugin from the Marketplace.")
+	UDEPRECATED_UOculusMR_Settings(const FObjectInitializer& ObjectInitializer);
 
 	/** Specify the distance to the camera which divide the background and foreground in MxR casting.
 	  * Set it to CR_TrackingReference to use the distance to the Tracking Reference, which works better
 	  * in the stationary experience. Set it to CR_Head would use the distance to the HMD, which works better
 	  * in the room scale experience.
 	  */
-	UPROPERTY(Category = OculusMR, EditAnywhere, BlueprintReadWrite)
+
+	UE_DEPRECATED(5.1, "OculusVR plugin is deprecated; please use the built-in OpenXR plugin or OculusXR plugin from the Marketplace.")
+	UPROPERTY(Category = OculusMR, EditAnywhere, BlueprintReadWrite, meta = (DeprecatedProperty))
 	EOculusMR_ClippingReference ClippingReference;
 
+
 	/** The casting viewports would use the same resolution of the camera which used in the calibration process. */
-	UPROPERTY(Category = OculusMR, EditAnywhere, BlueprintReadWrite)
+	UE_DEPRECATED(5.1, "OculusVR plugin is deprecated; please use the built-in OpenXR plugin or OculusXR plugin from the Marketplace.")
+	UPROPERTY(Category = OculusMR, EditAnywhere, BlueprintReadWrite, meta = (DeprecatedProperty))
 	bool bUseTrackedCameraResolution;
 
 	/** When bUseTrackedCameraResolution is false, the width of each casting viewport */
-	UPROPERTY(Category = OculusMR, EditAnywhere, BlueprintReadWrite)
+	UE_DEPRECATED(5.1, "OculusVR plugin is deprecated; please use the built-in OpenXR plugin or OculusXR plugin from the Marketplace.")
+	UPROPERTY(Category = OculusMR, EditAnywhere, BlueprintReadWrite, meta = (DeprecatedProperty))
 	int WidthPerView;
 
 	/** When bUseTrackedCameraResolution is false, the height of each casting viewport */
-	UPROPERTY(Category = OculusMR, EditAnywhere, BlueprintReadWrite)
+	UE_DEPRECATED(5.1, "OculusVR plugin is deprecated; please use the built-in OpenXR plugin or OculusXR plugin from the Marketplace.")
+	UPROPERTY(Category = OculusMR, EditAnywhere, BlueprintReadWrite, meta = (DeprecatedProperty))
 	int HeightPerView;
 
 	/** When CompositionMethod is External Composition, the latency of the casting output which could be adjusted to
 	  * match the camera latency in the external composition application */
-	UPROPERTY(Category = OculusMR, EditAnywhere, BlueprintReadWrite, meta = (UIMin = "0.0", UIMax = "0.1"))
+	UE_DEPRECATED(5.1, "OculusVR plugin is deprecated; please use the built-in OpenXR plugin or OculusXR plugin from the Marketplace.")
+	UPROPERTY(Category = OculusMR, EditAnywhere, BlueprintReadWrite, meta = (UIMin = "0.0", UIMax = "0.1", DeprecatedProperty))
 	float CastingLatency;
 
 	/** When CompositionMethod is External Composition, the color of the backdrop in the foreground view */
-	UPROPERTY(Category = OculusMR, EditAnywhere, BlueprintReadWrite)
+	UE_DEPRECATED(5.1, "OculusVR plugin is deprecated; please use the built-in OpenXR plugin or OculusXR plugin from the Marketplace.")
+	UPROPERTY(Category = OculusMR, EditAnywhere, BlueprintReadWrite, meta = (DeprecatedProperty))
 	FColor BackdropColor;
 
 	/** When CompositionMethod is Direct Composition, you could adjust this latency to delay the virtual
 	* hand movement by a small amount of time to match the camera latency */
-	UPROPERTY(Category = OculusMR, EditAnywhere, BlueprintReadWrite, meta = (UIMin = "0.0", UIMax = "0.5"))
+	UE_DEPRECATED(5.1, "OculusVR plugin is deprecated; please use the built-in OpenXR plugin or OculusXR plugin from the Marketplace.")
+	UPROPERTY(Category = OculusMR, EditAnywhere, BlueprintReadWrite, meta = (UIMin = "0.0", UIMax = "0.5", DeprecatedProperty))
 	float HandPoseStateLatency;
 
 	/** [Green-screen removal] Chroma Key Color. Apply when CompositionMethod is DirectComposition */
-	UPROPERTY(Category = OculusMR, EditAnywhere, BlueprintReadWrite)
+	UE_DEPRECATED(5.1, "OculusVR plugin is deprecated; please use the built-in OpenXR plugin or OculusXR plugin from the Marketplace.")
+	UPROPERTY(Category = OculusMR, EditAnywhere, BlueprintReadWrite, meta = (DeprecatedProperty))
 	FColor ChromaKeyColor;
 
 	/** [Green-screen removal] Chroma Key Similarity. Apply when CompositionMethod is DirectComposition */
-	UPROPERTY(Category = OculusMR, EditAnywhere, BlueprintReadWrite, meta = (UIMin = "0.0", UIMax = "1.0"))
+	UE_DEPRECATED(5.1, "OculusVR plugin is deprecated; please use the built-in OpenXR plugin or OculusXR plugin from the Marketplace.")
+	UPROPERTY(Category = OculusMR, EditAnywhere, BlueprintReadWrite, meta = (UIMin = "0.0", UIMax = "1.0", DeprecatedProperty))
 	float ChromaKeySimilarity;
 
 	/** [Green-screen removal] Chroma Key Smooth Range. Apply when CompositionMethod is DirectComposition */
-	UPROPERTY(Category = OculusMR, EditAnywhere, BlueprintReadWrite, meta = (UIMin = "0.0", UIMax = "0.2"))
+	UE_DEPRECATED(5.1, "OculusVR plugin is deprecated; please use the built-in OpenXR plugin or OculusXR plugin from the Marketplace.")
+	UPROPERTY(Category = OculusMR, EditAnywhere, BlueprintReadWrite, meta = (UIMin = "0.0", UIMax = "0.2", DeprecatedProperty))
 	float ChromaKeySmoothRange;
 
 	/** [Green-screen removal] Chroma Key Spill Range. Apply when CompositionMethod is DirectComposition */
-	UPROPERTY(Category = OculusMR, EditAnywhere, BlueprintReadWrite, meta = (UIMin = "0.0", UIMax = "0.2"))
+	UE_DEPRECATED(5.1, "OculusVR plugin is deprecated; please use the built-in OpenXR plugin or OculusXR plugin from the Marketplace.")
+	UPROPERTY(Category = OculusMR, EditAnywhere, BlueprintReadWrite, meta = (UIMin = "0.0", UIMax = "0.2", DeprecatedProperty))
 	float ChromaKeySpillRange;
 
 	/** Set the amount of post process effects in the MR view for external composition */
-	UPROPERTY(Category = OculusMR, EditAnywhere, BlueprintReadWrite)
+	UE_DEPRECATED(5.1, "OculusVR plugin is deprecated; please use the built-in OpenXR plugin or OculusXR plugin from the Marketplace.")
+	UPROPERTY(Category = OculusMR, EditAnywhere, BlueprintReadWrite, meta = (DeprecatedProperty))
 	EOculusMR_PostProcessEffects ExternalCompositionPostProcessEffects;
 
-
 	/** ExternalComposition: The casting window includes the background and foreground view
 	  * DirectComposition: The game scene would be composited with the camera frame directly
 	  */
-	UFUNCTION(BlueprintCallable, Category = OculusMR)
+	UE_DEPRECATED(5.1, "OculusVR plugin is deprecated; please use the built-in OpenXR plugin or OculusXR plugin from the Marketplace.")
+	UFUNCTION(BlueprintCallable, Category = OculusMR, meta = (DeprecatedFunction, DeprecationMessage = "The OculusVR plugin is deprecated."))
 	EOculusMR_CompositionMethod GetCompositionMethod() { return CompositionMethod; }
 
 	/** ExternalComposition: The casting window includes the background and foreground view
 	  * DirectComposition: The game scene would be composited with the camera frame directly
 	  */
-	UFUNCTION(BlueprintCallable, Category = OculusMR)
+	UE_DEPRECATED(5.1, "OculusVR plugin is deprecated; please use the built-in OpenXR plugin or OculusXR plugin from the Marketplace.")
+	UFUNCTION(BlueprintCallable, Category = OculusMR, meta = (DeprecatedFunction, DeprecationMessage = "The OculusVR plugin is deprecated."))
 	void SetCompositionMethod(EOculusMR_CompositionMethod val);
 
 	/** When CompositionMethod is DirectComposition, the physical camera device which provide the frame */
-	UFUNCTION(BlueprintCallable, Category = OculusMR)
+	UE_DEPRECATED(5.1, "OculusVR plugin is deprecated; please use the built-in OpenXR plugin or OculusXR plugin from the Marketplace.")
+	UFUNCTION(BlueprintCallable, Category = OculusMR, meta = (DeprecatedFunction, DeprecationMessage = "The OculusVR plugin is deprecated."))
 	EOculusMR_CameraDeviceEnum GetCapturingCamera() { return CapturingCamera; }
 
 	/** When CompositionMethod is DirectComposition, the physical camera device which provide the frame */
-	UFUNCTION(BlueprintCallable, Category = OculusMR)
+	UE_DEPRECATED(5.1, "OculusVR plugin is deprecated; please use the built-in OpenXR plugin or OculusXR plugin from the Marketplace.")
+	UFUNCTION(BlueprintCallable, Category = OculusMR, meta = (DeprecatedFunction, DeprecationMessage = "The OculusVR plugin is deprecated."))
 	void SetCapturingCamera(EOculusMR_CameraDeviceEnum val);
 
 	/** Is MRC on and off */
-	UFUNCTION(BlueprintCallable, Category = OculusMR)
+	UFUNCTION(BlueprintCallable, Category = OculusMR, meta = (DeprecatedFunction, DeprecationMessage = "The OculusVR plugin is deprecated."))
 	bool GetIsCasting() { return bIsCasting; }
 
 	/** Turns MRC on and off */
-	UFUNCTION(BlueprintCallable, Category = OculusMR)
+	UE_DEPRECATED(5.1, "OculusVR plugin is deprecated; please use the built-in OpenXR plugin or OculusXR plugin from the Marketplace.")
+	UFUNCTION(BlueprintCallable, Category = OculusMR, meta = (DeprecatedFunction, DeprecationMessage = "The OculusVR plugin is deprecated."))
 	void SetIsCasting(bool val);
 
 	/** Bind the casting camera to the calibrated external camera.
 	  * (Requires a calibrated external camera)
 	  */
-	UFUNCTION(BlueprintCallable, Category = OculusMR)
+	UE_DEPRECATED(5.1, "OculusVR plugin is deprecated; please use the built-in OpenXR plugin or OculusXR plugin from the Marketplace.")
+	UFUNCTION(BlueprintCallable, Category = OculusMR, meta = (DeprecatedFunction, DeprecationMessage = "The OculusVR plugin is deprecated."))
 	void BindToTrackedCameraIndexIfAvailable(int InTrackedCameraIndex);
 
-	UFUNCTION(BlueprintCallable, Category = OculusMR)
+	UE_DEPRECATED(5.1, "OculusVR plugin is deprecated; please use the built-in OpenXR plugin or OculusXR plugin from the Marketplace.")
+	UFUNCTION(BlueprintCallable, Category = OculusMR, meta = (DeprecatedFunction, DeprecationMessage = "The OculusVR plugin is deprecated."))
 	int GetBindToTrackedCameraIndex() { return BindToTrackedCameraIndex; }
 
 	/** Load settings from the config file */
-	UFUNCTION(BlueprintCallable, Category = OculusMR)
+	UE_DEPRECATED(5.1, "OculusVR plugin is deprecated; please use the built-in OpenXR plugin or OculusXR plugin from the Marketplace.")
+	UFUNCTION(BlueprintCallable, Category = OculusMR, meta = (DeprecatedFunction, DeprecationMessage = "The OculusVR plugin is deprecated."))
 	void LoadFromIni();
 
 	/** Save settings to the config file */
-	UFUNCTION(BlueprintCallable, Category = OculusMR)
+	UE_DEPRECATED(5.1, "OculusVR plugin is deprecated; please use the built-in OpenXR plugin or OculusXR plugin from the Marketplace.")
+	UFUNCTION(BlueprintCallable, Category = OculusMR, meta = (DeprecatedFunction, DeprecationMessage = "The OculusVR plugin is deprecated."))
 	void SaveToIni() const;
 
 private:
@@ -174,4 +199,6 @@ private:
 
 	// Give the OculusMR module access to the delegates so that 
 	friend class FOculusMRModule;
-};
\ No newline at end of file
+};
+
+PRAGMA_ENABLE_DEPRECATION_WARNINGS
