OculusVR deprecation: OculusInput module

Workaround needed for EBone, will follow-up in another CL

diff --git a/Engine/Plugins/Runtime/Oculus/OculusVR/Source/OculusInput/Private/OculusHandComponent.cpp b/Engine/Plugins/Runtime/Oculus/OculusVR/Source/OculusInput/Private/OculusHandComponent.cpp
--- a/Engine/Plugins/Runtime/Oculus/OculusVR/Source/OculusInput/Private/OculusHandComponent.cpp
+++ b/Engine/Plugins/Runtime/Oculus/OculusVR/Source/OculusInput/Private/OculusHandComponent.cpp
@@ -10,7 +10,9 @@
 
 const FQuat HandRootFixupRotation = FQuat(-0.5f, -0.5f, 0.5f, 0.5f);
 
-UOculusHandComponent::UOculusHandComponent(const FObjectInitializer& ObjectInitializer)
+PRAGMA_DISABLE_DEPRECATION_WARNINGS
+
+UDEPRECATED_UOculusHandComponent::UDEPRECATED_UOculusHandComponent(const FObjectInitializer& ObjectInitializer)
 	: Super(ObjectInitializer)
 {
 	PrimaryComponentTick.bCanEverTick = true;
@@ -28,7 +30,7 @@ UOculusHandComponent::UOculusHandComponent(const FObjectInitializer& ObjectIniti
 	}
 }
 
-void UOculusHandComponent::BeginPlay()
+void UDEPRECATED_UOculusHandComponent::BeginPlay()
 {
 	Super::BeginPlay();
 
@@ -45,11 +47,11 @@ void UOculusHandComponent::BeginPlay()
 	}
 }
 
-void UOculusHandComponent::InitializeSkeletalMesh()
+void UDEPRECATED_UOculusHandComponent::InitializeSkeletalMesh()
 {
 	if (RuntimeSkeletalMesh)
 	{
-		if (UOculusInputFunctionLibrary::GetHandSkeletalMesh(RuntimeSkeletalMesh, SkeletonType, MeshType))
+		if (UDEPRECATED_UOculusInputFunctionLibrary::GetHandSkeletalMesh(RuntimeSkeletalMesh, SkeletonType, MeshType))
 		{
 			SetSkinnedAssetAndUpdate(RuntimeSkeletalMesh, true);
 			if (MaterialOverride)
@@ -62,13 +64,13 @@ void UOculusHandComponent::InitializeSkeletalMesh()
 			// Initialize physics capsules on the runtime mesh
 			if (bInitializePhysics)
 			{
-				CollisionCapsules = UOculusInputFunctionLibrary::InitializeHandPhysics(SkeletonType, this);
+				CollisionCapsules = UDEPRECATED_UOculusInputFunctionLibrary::InitializeHandPhysics(SkeletonType, this);
 			}
 		}
 	}
 }
 
-void UOculusHandComponent::TickComponent(float DeltaTime, enum ELevelTick TickType, FActorComponentTickFunction* ThisTickFunction)
+void UDEPRECATED_UOculusHandComponent::TickComponent(float DeltaTime, enum ELevelTick TickType, FActorComponentTickFunction* ThisTickFunction)
 {
 	Super::TickComponent(DeltaTime, TickType, ThisTickFunction);
 
@@ -89,19 +91,19 @@ void UOculusHandComponent::TickComponent(float DeltaTime, enum ELevelTick TickTy
 	if (bHasAuthority)
 	{
 		bool bHidden = false;
-		if (UOculusInputFunctionLibrary::IsHandTrackingEnabled())
+		if (UDEPRECATED_UOculusInputFunctionLibrary::IsHandTrackingEnabled())
 		{
 			// Update Visibility based on Confidence
 			if (ConfidenceBehavior == EConfidenceBehavior::HideActor)
 			{
-				ETrackingConfidence TrackingConfidence = UOculusInputFunctionLibrary::GetTrackingConfidence(SkeletonType);
+				ETrackingConfidence TrackingConfidence = UDEPRECATED_UOculusInputFunctionLibrary::GetTrackingConfidence(SkeletonType);
 				bHidden |= TrackingConfidence != ETrackingConfidence::High;
 			}
 
 			// Update Hand Scale
 			if (bUpdateHandScale)
 			{
-				float NewScale = UOculusInputFunctionLibrary::GetHandScale(SkeletonType);
+				float NewScale = UDEPRECATED_UOculusInputFunctionLibrary::GetHandScale(SkeletonType);
 				SetRelativeScale3D(FVector(NewScale));
 			}
 
@@ -145,7 +147,7 @@ void UOculusHandComponent::TickComponent(float DeltaTime, enum ELevelTick TickTy
 	}
 }
 
-void UOculusHandComponent::UpdateBonePose()
+void UDEPRECATED_UOculusHandComponent::UpdateBonePose()
 {
 	if (bCustomHandMesh)
 	{
@@ -154,7 +156,7 @@ void UOculusHandComponent::UpdateBonePose()
 			// Set Root Bone Rotaiton
 			if (BoneElem.Key == EBone::Wrist_Root)
 			{
-				FQuat RootBoneRotation = UOculusInputFunctionLibrary::GetBoneRotation(SkeletonType, EBone::Wrist_Root);
+				FQuat RootBoneRotation = UDEPRECATED_UOculusInputFunctionLibrary::GetBoneRotation(SkeletonType, EBone::Wrist_Root);
 				RootBoneRotation *= HandRootFixupRotation;
 				RootBoneRotation.Normalize();
 				BoneSpaceTransforms[0].SetRotation(RootBoneRotation);
@@ -166,7 +168,7 @@ void UOculusHandComponent::UpdateBonePose()
 				int32 BoneIndex = GetSkinnedAsset()->GetRefSkeleton().FindBoneIndex(BoneElem.Value);
 				if (BoneIndex >= 0)
 				{
-					FQuat BoneRotation = UOculusInputFunctionLibrary::GetBoneRotation(SkeletonType, (EBone)BoneElem.Key);
+					FQuat BoneRotation = UDEPRECATED_UOculusInputFunctionLibrary::GetBoneRotation(SkeletonType, (EBone)BoneElem.Key);
 					BoneSpaceTransforms[BoneIndex].SetRotation(BoneRotation);
 				}
 			}
@@ -175,7 +177,7 @@ void UOculusHandComponent::UpdateBonePose()
 	else
 	{
 		// Set Root Bone Rotation
-		FQuat RootBoneRotation = UOculusInputFunctionLibrary::GetBoneRotation(SkeletonType, EBone::Wrist_Root);
+		FQuat RootBoneRotation = UDEPRECATED_UOculusInputFunctionLibrary::GetBoneRotation(SkeletonType, EBone::Wrist_Root);
 		RootBoneRotation *= HandRootFixupRotation;
 		RootBoneRotation.Normalize();
 		BoneSpaceTransforms[0].SetRotation(RootBoneRotation);
@@ -183,14 +185,14 @@ void UOculusHandComponent::UpdateBonePose()
 		// Set Remaining Bone Rotations
 		for (uint32 BoneIndex = 1; BoneIndex < (uint32)GetSkinnedAsset()->GetRefSkeleton().GetNum(); BoneIndex++)
 		{
-			FQuat BoneRotation = UOculusInputFunctionLibrary::GetBoneRotation(SkeletonType, (EBone)BoneIndex);
+			FQuat BoneRotation = UDEPRECATED_UOculusInputFunctionLibrary::GetBoneRotation(SkeletonType, (EBone)BoneIndex);
 			BoneSpaceTransforms[BoneIndex].SetRotation(BoneRotation);
 		}
 	}
 	MarkRefreshTransformDirty();
 }
 
-void UOculusHandComponent::SystemGesturePressed()
+void UDEPRECATED_UOculusHandComponent::SystemGesturePressed()
 {
 	if (SystemGestureBehavior == ESystemGestureBehavior::SwapMaterial)
 	{
@@ -205,7 +207,7 @@ void UOculusHandComponent::SystemGesturePressed()
 	}
 }
 
-void UOculusHandComponent::SystemGestureReleased()
+void UDEPRECATED_UOculusHandComponent::SystemGestureReleased()
 {
 	if (SystemGestureBehavior == ESystemGestureBehavior::SwapMaterial)
 	{
@@ -218,4 +220,6 @@ void UOculusHandComponent::SystemGestureReleased()
 			UE_LOG(LogTemp, Log, TEXT("System Gesture Behavior was set to Swap Material but no System Gesture Material was provided!"));
 		}
 	}
-}
\ No newline at end of file
+}
+
+PRAGMA_ENABLE_DEPRECATION_WARNINGS
diff --git a/Engine/Plugins/Runtime/Oculus/OculusVR/Source/OculusInput/Private/OculusHandTracking.cpp b/Engine/Plugins/Runtime/Oculus/OculusVR/Source/OculusInput/Private/OculusHandTracking.cpp
--- a/Engine/Plugins/Runtime/Oculus/OculusVR/Source/OculusInput/Private/OculusHandTracking.cpp
+++ b/Engine/Plugins/Runtime/Oculus/OculusVR/Source/OculusInput/Private/OculusHandTracking.cpp
@@ -20,6 +20,7 @@
 namespace OculusInput
 {
 
+PRAGMA_DISABLE_DEPRECATION_WARNINGS
 FQuat FOculusHandTracking::GetBoneRotation(const int32 ControllerIndex, const EOculusHandType DeviceHand, const EBone BoneId)
 {
 	IPlatformInputDeviceMapper& DeviceMapper = IPlatformInputDeviceMapper::Get();
@@ -667,4 +668,5 @@ FQuat FOculusHandTracking::OvrBoneQuatToFQuat(ovrpQuatf ovrpQuat)
 {
 	return FQuat(ovrpQuat.x, -ovrpQuat.y, ovrpQuat.z, -ovrpQuat.w);
 }
+PRAGMA_ENABLE_DEPRECATION_WARNINGS
 }
\ No newline at end of file
diff --git a/Engine/Plugins/Runtime/Oculus/OculusVR/Source/OculusInput/Private/OculusHandTracking.h b/Engine/Plugins/Runtime/Oculus/OculusVR/Source/OculusInput/Private/OculusHandTracking.h
--- a/Engine/Plugins/Runtime/Oculus/OculusVR/Source/OculusInput/Private/OculusHandTracking.h
+++ b/Engine/Plugins/Runtime/Oculus/OculusVR/Source/OculusInput/Private/OculusHandTracking.h
@@ -12,6 +12,8 @@
 
 DEFINE_LOG_CATEGORY_STATIC(LogOcHandTracking, Log, All);
 
+PRAGMA_DISABLE_DEPRECATION_WARNINGS
+
 //-------------------------------------------------------------------------------------------------
 // FOculusHandTracking
 //-------------------------------------------------------------------------------------------------
@@ -51,4 +53,6 @@ private:
 
 } // namespace OculusInput
 
+PRAGMA_ENABLE_DEPRECATION_WARNINGS
+
 #undef LOCTEXT_NAMESPACE
diff --git a/Engine/Plugins/Runtime/Oculus/OculusVR/Source/OculusInput/Private/OculusInput.cpp b/Engine/Plugins/Runtime/Oculus/OculusVR/Source/OculusInput/Private/OculusInput.cpp
--- a/Engine/Plugins/Runtime/Oculus/OculusVR/Source/OculusInput/Private/OculusInput.cpp
+++ b/Engine/Plugins/Runtime/Oculus/OculusVR/Source/OculusInput/Private/OculusInput.cpp
@@ -1059,7 +1059,7 @@ bool FOculusInput::GetControllerOrientationAndPosition( const int32 ControllerIn
 									}
 
 									auto bSuccess = true;
-									UOculusInputFunctionLibrary::HandMovementFilter.Broadcast(
+									UDEPRECATED_UOculusInputFunctionLibrary::HandMovementFilter.Broadcast(
 										DeviceHand,
 										&OutPosition,
 										&OutOrientation,
@@ -1078,7 +1078,7 @@ bool FOculusInput::GetControllerOrientationAndPosition( const int32 ControllerIn
 	}
 
 	auto bSuccess = false;
-	UOculusInputFunctionLibrary::HandMovementFilter.Broadcast(
+	UDEPRECATED_UOculusInputFunctionLibrary::HandMovementFilter.Broadcast(
 		DeviceHand,
 		&OutPosition,
 		&OutOrientation,
diff --git a/Engine/Plugins/Runtime/Oculus/OculusVR/Source/OculusInput/Private/OculusInputFunctionLibrary.cpp b/Engine/Plugins/Runtime/Oculus/OculusVR/Source/OculusInput/Private/OculusInputFunctionLibrary.cpp
--- a/Engine/Plugins/Runtime/Oculus/OculusVR/Source/OculusInput/Private/OculusInputFunctionLibrary.cpp
+++ b/Engine/Plugins/Runtime/Oculus/OculusVR/Source/OculusInput/Private/OculusInputFunctionLibrary.cpp
@@ -4,17 +4,20 @@
 #include "OculusHandTracking.h"
 #include "Logging/MessageLog.h"
 
+PRAGMA_DISABLE_DEPRECATION_WARNINGS
+
 //-------------------------------------------------------------------------------------------------
 // UOculusHandTrackingFunctionLibrary
 //-------------------------------------------------------------------------------------------------
-UOculusInputFunctionLibrary::UOculusInputFunctionLibrary(const FObjectInitializer& ObjectInitializer)
+UDEPRECATED_UOculusInputFunctionLibrary::UDEPRECATED_UOculusInputFunctionLibrary(const FObjectInitializer& ObjectInitializer)
 	: Super(ObjectInitializer)
 {
 }
 
-UOculusInputFunctionLibrary::FHandMovementFilterDelegate UOculusInputFunctionLibrary::HandMovementFilter;
+UDEPRECATED_UOculusInputFunctionLibrary::FHandMovementFilterDelegate UDEPRECATED_UOculusInputFunctionLibrary::HandMovementFilter;
+
 
-EOculusFinger UOculusInputFunctionLibrary::ConvertBoneToFinger(const EBone Bone)
+EOculusFinger UDEPRECATED_UOculusInputFunctionLibrary::ConvertBoneToFinger(const EBone Bone)
 {
 	switch (Bone)
 	{
@@ -50,47 +53,47 @@ EOculusFinger UOculusInputFunctionLibrary::ConvertBoneToFinger(const EBone Bone)
 	}
 }
 
-ETrackingConfidence UOculusInputFunctionLibrary::GetFingerTrackingConfidence(const EOculusHandType DeviceHand, const EOculusFinger Finger, const int32 ControllerIndex)
+ETrackingConfidence UDEPRECATED_UOculusInputFunctionLibrary::GetFingerTrackingConfidence(const EOculusHandType DeviceHand, const EOculusFinger Finger, const int32 ControllerIndex)
 {
 	return OculusInput::FOculusHandTracking::GetFingerTrackingConfidence(ControllerIndex, DeviceHand, (OculusInput::EOculusHandAxes)(uint8)Finger);
 }
 
-bool UOculusInputFunctionLibrary::GetHandSkeletalMesh(USkeletalMesh* HandSkeletalMesh, EOculusHandType SkeletonType, EOculusHandType MeshType, float WorldToMeters)
+bool UDEPRECATED_UOculusInputFunctionLibrary::GetHandSkeletalMesh(USkeletalMesh* HandSkeletalMesh, EOculusHandType SkeletonType, EOculusHandType MeshType, float WorldToMeters)
 {
 	return OculusInput::FOculusHandTracking::GetHandSkeletalMesh(HandSkeletalMesh, SkeletonType, MeshType, WorldToMeters);
 }
 
-TArray<FOculusCapsuleCollider> UOculusInputFunctionLibrary::InitializeHandPhysics(EOculusHandType SkeletonType, USkinnedMeshComponent* HandComponent, const float WorldToMeters)
+TArray<FOculusCapsuleCollider> UDEPRECATED_UOculusInputFunctionLibrary::InitializeHandPhysics(EOculusHandType SkeletonType, USkinnedMeshComponent* HandComponent, const float WorldToMeters)
 {
 	return OculusInput::FOculusHandTracking::InitializeHandPhysics(SkeletonType, HandComponent, WorldToMeters);
 }
 
-FQuat UOculusInputFunctionLibrary::GetBoneRotation(const EOculusHandType DeviceHand, const EBone BoneId, const int32 ControllerIndex)
+FQuat UDEPRECATED_UOculusInputFunctionLibrary::GetBoneRotation(const EOculusHandType DeviceHand, const EBone BoneId, const int32 ControllerIndex)
 {
 	return OculusInput::FOculusHandTracking::GetBoneRotation(ControllerIndex, DeviceHand, BoneId);
 }
 
-ETrackingConfidence UOculusInputFunctionLibrary::GetTrackingConfidence(const EOculusHandType DeviceHand, const int32 ControllerIndex)
+ETrackingConfidence UDEPRECATED_UOculusInputFunctionLibrary::GetTrackingConfidence(const EOculusHandType DeviceHand, const int32 ControllerIndex)
 {
 	return OculusInput::FOculusHandTracking::GetTrackingConfidence(ControllerIndex, DeviceHand);
 }
 
-FTransform UOculusInputFunctionLibrary::GetPointerPose(const EOculusHandType DeviceHand, const int32 ControllerIndex)
+FTransform UDEPRECATED_UOculusInputFunctionLibrary::GetPointerPose(const EOculusHandType DeviceHand, const int32 ControllerIndex)
 {
 	return OculusInput::FOculusHandTracking::GetPointerPose(ControllerIndex, DeviceHand);
 }
 
-bool UOculusInputFunctionLibrary::IsPointerPoseValid(const EOculusHandType DeviceHand, const int32 ControllerIndex)
+bool UDEPRECATED_UOculusInputFunctionLibrary::IsPointerPoseValid(const EOculusHandType DeviceHand, const int32 ControllerIndex)
 {
 	return OculusInput::FOculusHandTracking::IsPointerPoseValid(ControllerIndex, DeviceHand);
 }
 
-float UOculusInputFunctionLibrary::GetHandScale(const EOculusHandType DeviceHand, const int32 ControllerIndex)
+float UDEPRECATED_UOculusInputFunctionLibrary::GetHandScale(const EOculusHandType DeviceHand, const int32 ControllerIndex)
 {
 	return OculusInput::FOculusHandTracking::GetHandScale(ControllerIndex, DeviceHand);
 }
 
-EOculusHandType UOculusInputFunctionLibrary::GetDominantHand(const int32 ControllerIndex)
+EOculusHandType UDEPRECATED_UOculusInputFunctionLibrary::GetDominantHand(const int32 ControllerIndex)
 {
 	EOculusHandType DominantHand = EOculusHandType::None;
 	if (OculusInput::FOculusHandTracking::IsHandDominant(ControllerIndex, EOculusHandType::HandLeft))
@@ -104,19 +107,20 @@ EOculusHandType UOculusInputFunctionLibrary::GetDominantHand(const int32 Control
 	return DominantHand;
 }
 
-bool UOculusInputFunctionLibrary::IsHandTrackingEnabled()
+bool UDEPRECATED_UOculusInputFunctionLibrary::IsHandTrackingEnabled()
 {
 	return OculusInput::FOculusHandTracking::IsHandTrackingEnabled();
 }
 
-bool UOculusInputFunctionLibrary::IsHandPositionValid(const EOculusHandType DeviceHand, const int32 ControllerIndex)
+bool UDEPRECATED_UOculusInputFunctionLibrary::IsHandPositionValid(const EOculusHandType DeviceHand, const int32 ControllerIndex)
 {
 	return OculusInput::FOculusHandTracking::IsHandPositionValid(ControllerIndex, DeviceHand);
 }
 
-FString UOculusInputFunctionLibrary::GetBoneName(EBone BoneId)
+FString UDEPRECATED_UOculusInputFunctionLibrary::GetBoneName(EBone BoneId)
 {
 	uint32 ovrBoneId = OculusInput::FOculusHandTracking::ToOvrBone(BoneId);
 	return OculusInput::FOculusHandTracking::GetBoneName(ovrBoneId);
 }
+PRAGMA_ENABLE_DEPRECATION_WARNINGS
 
diff --git a/Engine/Plugins/Runtime/Oculus/OculusVR/Source/OculusInput/Private/OculusInputModule.h b/Engine/Plugins/Runtime/Oculus/OculusVR/Source/OculusInput/Private/OculusInputModule.h
--- a/Engine/Plugins/Runtime/Oculus/OculusVR/Source/OculusInput/Private/OculusInputModule.h
+++ b/Engine/Plugins/Runtime/Oculus/OculusVR/Source/OculusInput/Private/OculusInputModule.h
@@ -14,6 +14,8 @@
 
 #if OCULUS_INPUT_SUPPORTED_PLATFORMS
 
+PRAGMA_DISABLE_DEPRECATION_WARNINGS
+
 namespace OculusInput
 {
 	class FOculusInput;
@@ -48,6 +50,8 @@ class FOculusInputModule : public FDefaultModuleImpl
 	};
 };
 
+PRAGMA_ENABLE_DEPRECATION_WARNINGS
+
 #endif	// OCULUS_INPUT_SUPPORTED_PLATFORMS
 
 
diff --git a/Engine/Plugins/Runtime/Oculus/OculusVR/Source/OculusInput/Private/OculusInputState.h b/Engine/Plugins/Runtime/Oculus/OculusVR/Source/OculusInput/Private/OculusInputState.h
--- a/Engine/Plugins/Runtime/Oculus/OculusVR/Source/OculusInput/Private/OculusInputState.h
+++ b/Engine/Plugins/Runtime/Oculus/OculusVR/Source/OculusInput/Private/OculusInputState.h
@@ -9,6 +9,8 @@
 #include "OculusInputFunctionLibrary.h"
 #include "GenericPlatform/GenericApplicationMessageHandler.h"
 
+PRAGMA_DISABLE_DEPRECATION_WARNINGS
+
 namespace OculusInput
 {
 
@@ -538,4 +540,6 @@ struct FOculusControllerPair
 
 } // namespace OculusInput
 
+PRAGMA_ENABLE_DEPRECATION_WARNINGS
+
 #endif	// OCULUS_INPUT_SUPPORTED_PLATFORMS
diff --git a/Engine/Plugins/Runtime/Oculus/OculusVR/Source/OculusInput/Public/IOculusInputModule.h b/Engine/Plugins/Runtime/Oculus/OculusVR/Source/OculusInput/Public/IOculusInputModule.h
--- a/Engine/Plugins/Runtime/Oculus/OculusVR/Source/OculusInput/Public/IOculusInputModule.h
+++ b/Engine/Plugins/Runtime/Oculus/OculusVR/Source/OculusInput/Public/IOculusInputModule.h
@@ -11,6 +11,7 @@
  * The public interface to this module.  In most cases, this interface is only public to sibling modules 
  * within this plugin.
  */
+class UE_DEPRECATED(5.1, "OculusVR plugin is deprecated; please use the built-in OpenXR plugin or OculusXR plugin from the Marketplace.") IOculusInputModule;
 class IOculusInputModule : public IInputDeviceModule
 {
 
diff --git a/Engine/Plugins/Runtime/Oculus/OculusVR/Source/OculusInput/Public/OculusHandComponent.h b/Engine/Plugins/Runtime/Oculus/OculusVR/Source/OculusInput/Public/OculusHandComponent.h
--- a/Engine/Plugins/Runtime/Oculus/OculusVR/Source/OculusInput/Public/OculusHandComponent.h
+++ b/Engine/Plugins/Runtime/Oculus/OculusVR/Source/OculusInput/Public/OculusHandComponent.h
@@ -5,15 +5,17 @@
 #include "Components/PoseableMeshComponent.h"
 #include "OculusHandComponent.generated.h"
 
+PRAGMA_DISABLE_DEPRECATION_WARNINGS
+
 UENUM(BlueprintType)
-enum class EConfidenceBehavior : uint8
+enum class UE_DEPRECATED(5.1, "OculusVR plugin is deprecated; please use the built-in OpenXR plugin or OculusXR plugin from the Marketplace.") EConfidenceBehavior : uint8
 {
 	None,
 	HideActor
 };
 
 UENUM(BlueprintType)
-enum class ESystemGestureBehavior : uint8
+enum class UE_DEPRECATED(5.1, "OculusVR plugin is deprecated; please use the built-in OpenXR plugin or OculusXR plugin from the Marketplace.") ESystemGestureBehavior : uint8
 {
 	None,
 	SwapMaterial
@@ -21,8 +23,8 @@ enum class ESystemGestureBehavior : uint8
 
 OCULUSINPUT_API extern const FQuat HandRootFixupRotation;
 
-UCLASS(Blueprintable, meta = (BlueprintSpawnableComponent), ClassGroup = OculusHand)
-class OCULUSINPUT_API UOculusHandComponent : public UPoseableMeshComponent
+UCLASS(Blueprintable, meta = (BlueprintSpawnableComponent), ClassGroup = OculusHand, deprecated, meta = (DeprecationMessage = "OculusVR plugin is deprecated; please use the built-in OpenXR plugin or OculusXR plugin from the Marketplace."))
+class OCULUSINPUT_API UDEPRECATED_UOculusHandComponent : public UPoseableMeshComponent
 {
 	GENERATED_UCLASS_BODY()
 
@@ -32,47 +34,58 @@ public:
 	virtual void TickComponent(float DeltaTime, enum ELevelTick TickType, FActorComponentTickFunction* ThisTickFunction) override;
 
 	/** The hand skeleton that will be loaded */
-	UPROPERTY(EditDefaultsOnly, BlueprintReadWrite, Category = "HandProperties")
+	UE_DEPRECATED(5.1, "OculusVR plugin is deprecated; please use the built-in OpenXR plugin or OculusXR plugin from the Marketplace.")
+	UPROPERTY(EditDefaultsOnly, BlueprintReadWrite, Category = "HandProperties", meta = (DeprecatedProperty))
 	EOculusHandType SkeletonType;
 
 	/** The hand mesh that will be applied to the skeleton */
-	UPROPERTY(EditDefaultsOnly, BlueprintReadWrite, Category = "HandProperties")
+	UE_DEPRECATED(5.1, "OculusVR plugin is deprecated; please use the built-in OpenXR plugin or OculusXR plugin from the Marketplace.")
+	UPROPERTY(EditDefaultsOnly, BlueprintReadWrite, Category = "HandProperties", meta = (DeprecatedProperty))
 	EOculusHandType MeshType;
 
 	/** Behavior for when hand tracking loses high confidence tracking */
-	UPROPERTY(EditDefaultsOnly, BlueprintReadWrite, Category = "HandProperties")
+	UE_DEPRECATED(5.1, "OculusVR plugin is deprecated; please use the built-in OpenXR plugin or OculusXR plugin from the Marketplace.")
+	UPROPERTY(EditDefaultsOnly, BlueprintReadWrite, Category = "HandProperties", meta = (DeprecatedProperty))
 	EConfidenceBehavior ConfidenceBehavior = EConfidenceBehavior::HideActor;
 
 	/** Behavior for when the system gesture is actived */
-	UPROPERTY(EditDefaultsOnly, BlueprintReadWrite, Category = "HandProperties")
+	UE_DEPRECATED(5.1, "OculusVR plugin is deprecated; please use the built-in OpenXR plugin or OculusXR plugin from the Marketplace.")
+	UPROPERTY(EditDefaultsOnly, BlueprintReadWrite, Category = "HandProperties", meta = (DeprecatedProperty))
 	ESystemGestureBehavior SystemGestureBehavior = ESystemGestureBehavior::SwapMaterial;
 
 	/** Material that gets applied to the hands when the system gesture is active */
-	UPROPERTY(EditDefaultsOnly, BlueprintReadWrite, Category = "HandProperties")
+	UE_DEPRECATED(5.1, "OculusVR plugin is deprecated; please use the built-in OpenXR plugin or OculusXR plugin from the Marketplace.")
+	UPROPERTY(EditDefaultsOnly, BlueprintReadWrite, Category = "HandProperties", meta = (DeprecatedProperty))
 	TObjectPtr<class UMaterialInterface> SystemGestureMaterial;
 
 	/** Whether or not to initialize physics capsules on the skeletal mesh */
-	UPROPERTY(EditDefaultsOnly, BlueprintReadWrite, Category = "HandProperties")
+	UE_DEPRECATED(5.1, "OculusVR plugin is deprecated; please use the built-in OpenXR plugin or OculusXR plugin from the Marketplace.")
+	UPROPERTY(EditDefaultsOnly, BlueprintReadWrite, Category = "HandProperties", meta = (DeprecatedProperty))
 	bool bInitializePhysics;
 
 	/** Whether or not the hand scale should update based on values from the runtime to match the users hand scale */
-	UPROPERTY(EditDefaultsOnly, BlueprintReadWrite, Category = "HandProperties")
+	UE_DEPRECATED(5.1, "OculusVR plugin is deprecated; please use the built-in OpenXR plugin or OculusXR plugin from the Marketplace.")
+	UPROPERTY(EditDefaultsOnly, BlueprintReadWrite, Category = "HandProperties", meta = (DeprecatedProperty))
 	bool bUpdateHandScale;
 
 	/** Material override for the runtime skeletal mesh */
-	UPROPERTY(EditDefaultsOnly, BlueprintReadWrite, Category = "HandProperties")
+	UE_DEPRECATED(5.1, "OculusVR plugin is deprecated; please use the built-in OpenXR plugin or OculusXR plugin from the Marketplace.")
+	UPROPERTY(EditDefaultsOnly, BlueprintReadWrite, Category = "HandProperties", meta = (DeprecatedProperty))
 	TObjectPtr<class UMaterialInterface> MaterialOverride;
 
 	/** Bone mapping for custom hand skeletal meshes */
-	UPROPERTY(EditAnywhere, BlueprintReadWrite, Category = "CustomSkeletalMesh")
+	UE_DEPRECATED(5.1, "OculusVR plugin is deprecated; please use the built-in OpenXR plugin or OculusXR plugin from the Marketplace.")
+	UPROPERTY(EditAnywhere, BlueprintReadWrite, Category = "CustomSkeletalMesh", meta = (DeprecatedProperty))
 	TMap<EBone, FName> BoneNameMappings;
 
 	/** List of capsule colliders created for the skeletal mesh */
-	UPROPERTY(BlueprintReadOnly, Category = "HandProperties")
+	UE_DEPRECATED(5.1, "OculusVR plugin is deprecated; please use the built-in OpenXR plugin or OculusXR plugin from the Marketplace.")
+	UPROPERTY(BlueprintReadOnly, Category = "HandProperties", meta = (DeprecatedProperty))
 	TArray<FOculusCapsuleCollider> CollisionCapsules;
 
 	/** Whether or not the runtime skeletal mesh has been loaded and initialized */
-	UPROPERTY(BlueprintReadOnly, Category = "HandProperties")
+	UE_DEPRECATED(5.1, "OculusVR plugin is deprecated; please use the built-in OpenXR plugin or OculusXR plugin from the Marketplace.")
+	UPROPERTY(BlueprintReadOnly, Category = "HandProperties", meta = (DeprecatedProperty))
 	bool bSkeletalMeshInitialized = false;
 
 protected:
@@ -96,4 +109,6 @@ private:
 	void InitializeSkeletalMesh();
 
 	void UpdateBonePose();
-};
\ No newline at end of file
+};
+
+PRAGMA_ENABLE_DEPRECATION_WARNINGS
diff --git a/Engine/Plugins/Runtime/Oculus/OculusVR/Source/OculusInput/Public/OculusInputFunctionLibrary.h b/Engine/Plugins/Runtime/Oculus/OculusVR/Source/OculusInput/Public/OculusInputFunctionLibrary.h
--- a/Engine/Plugins/Runtime/Oculus/OculusVR/Source/OculusInput/Public/OculusInputFunctionLibrary.h
+++ b/Engine/Plugins/Runtime/Oculus/OculusVR/Source/OculusInput/Public/OculusInputFunctionLibrary.h
@@ -8,8 +8,10 @@
 
 #include "OculusInputFunctionLibrary.generated.h"
 
+PRAGMA_DISABLE_DEPRECATION_WARNINGS
+
 UENUM(BlueprintType)
-enum class EOculusHandType : uint8
+enum class UE_DEPRECATED(5.1, "OculusVR plugin is deprecated; please use the built-in OpenXR plugin or OculusXR plugin from the Marketplace.") EOculusHandType : uint8
 {
 	None,
 	HandLeft,
@@ -17,14 +19,14 @@ enum class EOculusHandType : uint8
 };
 
 UENUM(BlueprintType)
-enum class ETrackingConfidence : uint8
+enum class UE_DEPRECATED(5.1, "OculusVR plugin is deprecated; please use the built-in OpenXR plugin or OculusXR plugin from the Marketplace.") ETrackingConfidence : uint8
 {
 	Low,
 	High
 };
 
 UENUM(BlueprintType)
-enum class EOculusFinger : uint8
+enum class UE_DEPRECATED(5.1, "OculusVR plugin is deprecated; please use the built-in OpenXR plugin or OculusXR plugin from the Marketplace.") EOculusFinger : uint8
 {
 	Thumb,
 	Index,
@@ -34,11 +36,13 @@ enum class EOculusFinger : uint8
 	Invalid
 };
 
+// This is deprecated, but explicit deprecation causes issues with Android non-adaptive unity build.
+// We can follow-up once we investigate the build issue a bit more.
 /**
 * EBone is enum representing the Bone Ids that come from the Oculus Runtime. 
 */
 UENUM(BlueprintType)
-enum class EBone : uint8
+enum class /*UE_DEPRECATED(5.1, "OculusVR plugin is deprecated; please use the built-in OpenXR plugin or OculusXR plugin from the Marketplace.")*/ EBone : uint8
 {
 	Wrist_Root UMETA(DisplayName = "Wrist Root"),
 	Hand_Start = Wrist_Root UMETA(DisplayName = "Hand Start"),
@@ -78,7 +82,7 @@ enum class EBone : uint8
 * @var BoneIndex	The Bone that this collision capsule is parented to. Corresponds to the EBone enum.
 *
 */
-USTRUCT(BlueprintType)
+USTRUCT(BlueprintType, meta = (Deprecated = "5.1"))
 struct OCULUSINPUT_API FOculusCapsuleCollider
 {
 	GENERATED_BODY()
@@ -91,16 +95,17 @@ public:
 	EBone BoneId = EBone::Wrist_Root;
 };
 
-UCLASS()
-class OCULUSINPUT_API UOculusInputFunctionLibrary : public UBlueprintFunctionLibrary
+UCLASS(deprecated, meta = (DeprecationMessage = "OculusVR plugin is deprecated; please use the built-in OpenXR plugin or OculusXR plugin from the Marketplace."))
+class OCULUSINPUT_API UDEPRECATED_UOculusInputFunctionLibrary : public UBlueprintFunctionLibrary
 {
 	GENERATED_UCLASS_BODY()
 
 public:
 
-	UFUNCTION(BlueprintCallable, Category = "OculusLibrary|HandTracking")
+	UE_DEPRECATED(5.1, "OculusVR plugin is deprecated; please use the built-in OpenXR plugin or OculusXR plugin from the Marketplace.")
+	UFUNCTION(BlueprintCallable, Category = "OculusLibrary|HandTracking", meta = (DeprecatedFunction, DeprecationMessage = "The OculusVR plugin is deprecated."))
 	static EOculusFinger ConvertBoneToFinger(const EBone Bone);
-	
+
 	DECLARE_MULTICAST_DELEGATE_FourParams(FHandMovementFilterDelegate, EControllerHand, FVector*, FRotator*, bool*);
 	static FHandMovementFilterDelegate HandMovementFilter; /// Called to modify Hand position and orientation whenever it is queried
 	
@@ -112,7 +117,8 @@ public:
 	 * @param MeshType					(in) The mesh type that will be used for generating the hand mesh
 	 * @param WorldTometers				(in) Optional change to the world to meters conversion value
 	 */
-	UFUNCTION(BlueprintCallable, Category = "OculusLibrary|HandTracking")
+	UE_DEPRECATED(5.1, "OculusVR plugin is deprecated; please use the built-in OpenXR plugin or OculusXR plugin from the Marketplace.")
+	UFUNCTION(BlueprintCallable, Category = "OculusLibrary|HandTracking", meta = (DeprecatedFunction, DeprecationMessage = "The OculusVR plugin is deprecated."))
 	static bool GetHandSkeletalMesh(USkeletalMesh* HandSkeletalMesh, EOculusHandType SkeletonType, EOculusHandType MeshType, const float WorldToMeters = 100.0f);
 
 	/**
@@ -122,7 +128,8 @@ public:
 	 * @param HandComponent				(in) The skinned mesh component that the capsules will be attached to
 	 * @param WorldTometers				(in) Optional change to the world to meters conversion value
 	 */
-	UFUNCTION(BlueprintCallable, Category = "OculusLibrary|HandTracking")
+	UE_DEPRECATED(5.1, "OculusVR plugin is deprecated; please use the built-in OpenXR plugin or OculusXR plugin from the Marketplace.")
+	UFUNCTION(BlueprintCallable, Category = "OculusLibrary|HandTracking", meta = (DeprecatedFunction, DeprecationMessage = "The OculusVR plugin is deprecated."))
 	static TArray<FOculusCapsuleCollider> InitializeHandPhysics(EOculusHandType SkeletonType, USkinnedMeshComponent* HandComponent, const float WorldToMeters = 100.0f);
 
 	/**
@@ -132,7 +139,8 @@ public:
 	 * @param BoneId					(in) The specific bone to get the rotation from
 	 * @param ControllerIndex			(in) Optional different controller index
 	 */
-	UFUNCTION(BlueprintPure, Category = "OculusLibrary|HandTracking")
+	UE_DEPRECATED(5.1, "OculusVR plugin is deprecated; please use the built-in OpenXR plugin or OculusXR plugin from the Marketplace.")
+	UFUNCTION(BlueprintPure, Category = "OculusLibrary|HandTracking", meta = (DeprecatedFunction, DeprecationMessage = "The OculusVR plugin is deprecated."))
 	static FQuat GetBoneRotation(const EOculusHandType DeviceHand, const EBone BoneId, const int32 ControllerIndex = 0);
 	
 	/**
@@ -141,7 +149,8 @@ public:
 	 * @param DeviceHand				(in) The hand to get the pointer pose from
 	 * @param ControllerIndex			(in) Optional different controller index
 	 */
-	UFUNCTION(BlueprintPure, Category = "OculusLibrary|HandTracking")
+	UE_DEPRECATED(5.1, "OculusVR plugin is deprecated; please use the built-in OpenXR plugin or OculusXR plugin from the Marketplace.")
+	UFUNCTION(BlueprintPure, Category = "OculusLibrary|HandTracking", meta = (DeprecatedFunction, DeprecationMessage = "The OculusVR plugin is deprecated."))
 	static FTransform GetPointerPose(const EOculusHandType DeviceHand, const int32 ControllerIndex = 0);
 
 	/**
@@ -150,7 +159,8 @@ public:
 	 * @param DeviceHand				(in) The hand to get the pointer status from
 	 * @param ControllerIndex			(in) Optional different controller index
 	 */
-	UFUNCTION(BlueprintPure, Category = "OculusLibrary|HandTracking")
+	UE_DEPRECATED(5.1, "OculusVR plugin is deprecated; please use the built-in OpenXR plugin or OculusXR plugin from the Marketplace.")
+	UFUNCTION(BlueprintPure, Category = "OculusLibrary|HandTracking", meta = (DeprecatedFunction, DeprecationMessage = "The OculusVR plugin is deprecated."))
 	static bool IsPointerPoseValid(const EOculusHandType DeviceHand, const int32 ControllerIndex = 0);
 
 	/**
@@ -159,7 +169,8 @@ public:
 	 * @param DeviceHand				(in) The hand to get tracking confidence of
 	 * @param ControllerIndex			(in) Optional different controller index
 	 */
-	UFUNCTION(BlueprintPure, Category = "OculusLibrary|HandTracking")
+	UE_DEPRECATED(5.1, "OculusVR plugin is deprecated; please use the built-in OpenXR plugin or OculusXR plugin from the Marketplace.")
+	UFUNCTION(BlueprintPure, Category = "OculusLibrary|HandTracking", meta = (DeprecatedFunction, DeprecationMessage = "The OculusVR plugin is deprecated."))
 	static ETrackingConfidence GetTrackingConfidence(const EOculusHandType DeviceHand, const int32 ControllerIndex = 0);
 
 	/**
@@ -169,7 +180,8 @@ public:
 	* @param ControllerIndex			(in) Optional different controller index
 	* @param Finger			(in) The finger to get tracking confidence of
 	*/
-	UFUNCTION(BlueprintPure, Category = "OculusLibrary|HandTracking")
+	UE_DEPRECATED(5.1, "OculusVR plugin is deprecated; please use the built-in OpenXR plugin or OculusXR plugin from the Marketplace.")
+	UFUNCTION(BlueprintPure, Category = "OculusLibrary|HandTracking", meta = (DeprecatedFunction, DeprecationMessage = "The OculusVR plugin is deprecated."))
 	static ETrackingConfidence GetFingerTrackingConfidence(const EOculusHandType DeviceHand, const EOculusFinger Finger, const int32 ControllerIndex = 0);
 
 	/**
@@ -178,7 +190,8 @@ public:
 	 * @param DeviceHand				(in) The hand to get scale of
 	 * @param ControllerIndex			(in) Optional different controller index
 	 */
-	UFUNCTION(BlueprintPure, Category = "OculusLibrary|HandTracking")
+	UE_DEPRECATED(5.1, "OculusVR plugin is deprecated; please use the built-in OpenXR plugin or OculusXR plugin from the Marketplace.")
+	UFUNCTION(BlueprintPure, Category = "OculusLibrary|HandTracking", meta = (DeprecatedFunction, DeprecationMessage = "The OculusVR plugin is deprecated."))
 	static float GetHandScale(const EOculusHandType DeviceHand, const int32 ControllerIndex = 0);
 
 	/**
@@ -186,13 +199,15 @@ public:
 	 *
 	 * @param ControllerIndex			(in) Optional different controller index
 	 */
-	UFUNCTION(BlueprintPure, Category = "OculusLibrary|HandTracking")
+	UE_DEPRECATED(5.1, "OculusVR plugin is deprecated; please use the built-in OpenXR plugin or OculusXR plugin from the Marketplace.")
+	UFUNCTION(BlueprintPure, Category = "OculusLibrary|HandTracking", meta = (DeprecatedFunction, DeprecationMessage = "The OculusVR plugin is deprecated."))
 	static EOculusHandType GetDominantHand(const int32 ControllerIndex = 0);
 
 	/**
 	 * Check if hand tracking is enabled currently
 	 */
-	UFUNCTION(BlueprintPure, Category = "OculusLibrary|HandTracking")
+	UE_DEPRECATED(5.1, "OculusVR plugin is deprecated; please use the built-in OpenXR plugin or OculusXR plugin from the Marketplace.")
+	UFUNCTION(BlueprintPure, Category = "OculusLibrary|HandTracking", meta = (DeprecatedFunction, DeprecationMessage = "The OculusVR plugin is deprecated."))
 	static bool IsHandTrackingEnabled();
 
 	/**
@@ -201,7 +216,8 @@ public:
 	* @param DeviceHand				(in) The hand to get the position from
 	* @param ControllerIndex			(in) Optional different controller index
 	*/
-	UFUNCTION(BlueprintPure, Category = "OculusLibrary|HandTracking")
+	UE_DEPRECATED(5.1, "OculusVR plugin is deprecated; please use the built-in OpenXR plugin or OculusXR plugin from the Marketplace.")
+	UFUNCTION(BlueprintPure, Category = "OculusLibrary|HandTracking", meta = (DeprecatedFunction, DeprecationMessage = "The OculusVR plugin is deprecated."))
 	static bool IsHandPositionValid(const EOculusHandType DeviceHand, const int32 ControllerIndex = 0);
 
 	/**
@@ -209,7 +225,9 @@ public:
 	 *
 	 * @param BoneIndex					(in) Bone index to get the name of
 	 */
-	UFUNCTION(BlueprintPure, Category = "OculusLibrary|HandTracking")
+	UE_DEPRECATED(5.1, "OculusVR plugin is deprecated; please use the built-in OpenXR plugin or OculusXR plugin from the Marketplace.")
+	UFUNCTION(BlueprintPure, Category = "OculusLibrary|HandTracking", meta = (DeprecatedFunction, DeprecationMessage = "The OculusVR plugin is deprecated."))
 	static FString GetBoneName(EBone BoneId);
 };
+PRAGMA_ENABLE_DEPRECATION_WARNINGS
 
