OculusVR deprecation: OculusHMD module

diff --git a/Engine/Plugins/Runtime/Oculus/OculusVR/Source/OculusEditor/Private/OculusBuildAnalytics.cpp b/Engine/Plugins/Runtime/Oculus/OculusVR/Source/OculusEditor/Private/OculusBuildAnalytics.cpp
--- a/Engine/Plugins/Runtime/Oculus/OculusVR/Source/OculusEditor/Private/OculusBuildAnalytics.cpp
+++ b/Engine/Plugins/Runtime/Oculus/OculusVR/Source/OculusEditor/Private/OculusBuildAnalytics.cpp
@@ -7,6 +7,7 @@
 
 FOculusBuildAnalytics* FOculusBuildAnalytics::instance = 0;
 
+PRAGMA_DISABLE_DEPRECATION_WARNINGS
 FOculusBuildAnalytics* FOculusBuildAnalytics::GetInstance()
 {
 	if (IOculusHMDModule::IsAvailable())
@@ -311,3 +312,4 @@ void FOculusBuildAnalytics::SendBuildCompleteEvent(float TotalTime)
 	FOculusHMDModule::GetPluginWrapper().AddCustomMetadata("build_step_count", TCHAR_TO_ANSI(*FString::FromInt(BuildStepCount)));
 	FOculusHMDModule::GetPluginWrapper().SendEvent2("build_complete", TCHAR_TO_ANSI(*FString::SanitizeFloat(TotalTime)), "ovrbuild");
 }
+PRAGMA_ENABLE_DEPRECATION_WARNINGS
diff --git a/Engine/Plugins/Runtime/Oculus/OculusVR/Source/OculusEditor/Private/OculusEditorModule.cpp b/Engine/Plugins/Runtime/Oculus/OculusVR/Source/OculusEditor/Private/OculusEditorModule.cpp
--- a/Engine/Plugins/Runtime/Oculus/OculusVR/Source/OculusEditor/Private/OculusEditorModule.cpp
+++ b/Engine/Plugins/Runtime/Oculus/OculusVR/Source/OculusEditor/Private/OculusEditorModule.cpp
@@ -24,6 +24,8 @@
 
 #define LOCTEXT_NAMESPACE "OculusEditor"
 
+PRAGMA_DISABLE_DEPRECATION_WARNINGS
+
 const FName FOculusEditorModule::OculusPerfTabName = FName("OculusPerfCheck");
 const FName FOculusEditorModule::OculusPlatToolTabName = FName("OculusPlatormTool");
 
@@ -131,11 +133,11 @@ void FOculusEditorModule::RegisterSettings()
 		SettingsModule->RegisterSettings("Project", "Plugins", "OculusVR",
 			LOCTEXT("RuntimeSettingsName", "OculusVR"),
 			LOCTEXT("RuntimeSettingsDescription", "Configure the OculusVR plugin"),
-			GetMutableDefault<UOculusHMDRuntimeSettings>()
+			GetMutableDefault<UDEPRECATED_UOculusHMDRuntimeSettings>()
 		);
 
 		FPropertyEditorModule& PropertyModule = FModuleManager::GetModuleChecked<FPropertyEditorModule>("PropertyEditor");
-		PropertyModule.RegisterCustomClassLayout(UOculusHMDRuntimeSettings::StaticClass()->GetFName(), FOnGetDetailCustomizationInstance::CreateStatic(&FOculusHMDSettingsDetailsCustomization::MakeInstance));
+		PropertyModule.RegisterCustomClassLayout(UDEPRECATED_UOculusHMDRuntimeSettings::StaticClass()->GetFName(), FOnGetDetailCustomizationInstance::CreateStatic(&FOculusHMDSettingsDetailsCustomization::MakeInstance));
 	}
 }
 
@@ -306,6 +308,8 @@ void FOculusHMDSettingsDetailsCustomization::CustomizeDetails(IDetailLayoutBuild
 		];
 }
 
+PRAGMA_ENABLE_DEPRECATION_WARNINGS
+
 //////////////////////////////////////////////////////////////////////////
 
 IMPLEMENT_MODULE(FOculusEditorModule, OculusEditor);
diff --git a/Engine/Plugins/Runtime/Oculus/OculusVR/Source/OculusEditor/Private/OculusToolWidget.cpp b/Engine/Plugins/Runtime/Oculus/OculusVR/Source/OculusEditor/Private/OculusToolWidget.cpp
--- a/Engine/Plugins/Runtime/Oculus/OculusVR/Source/OculusEditor/Private/OculusToolWidget.cpp
+++ b/Engine/Plugins/Runtime/Oculus/OculusVR/Source/OculusEditor/Private/OculusToolWidget.cpp
@@ -1009,15 +1009,15 @@ EVisibility SOculusToolWidget::StartInVRVisibility(FName tag) const
 
 FReply SOculusToolWidget::SupportDashEnable(bool text)
 {
-	UOculusHMDRuntimeSettings* Settings = GetMutableDefault<UOculusHMDRuntimeSettings>();
+	UDEPRECATED_UOculusHMDRuntimeSettings* Settings = GetMutableDefault<UDEPRECATED_UOculusHMDRuntimeSettings>();
 	Settings->bSupportsDash = true;
-	Settings->UpdateSinglePropertyInConfigFile(Settings->GetClass()->FindPropertyByName(GET_MEMBER_NAME_CHECKED(UOculusHMDRuntimeSettings, bSupportsDash)), Settings->GetDefaultConfigFilename());
+	Settings->UpdateSinglePropertyInConfigFile(Settings->GetClass()->FindPropertyByName(GET_MEMBER_NAME_CHECKED(UDEPRECATED_UOculusHMDRuntimeSettings, bSupportsDash)), Settings->GetDefaultConfigFilename());
 	return FReply::Handled();
 }
 
 EVisibility SOculusToolWidget::SupportDashVisibility(FName tag) const
 {
-	const UOculusHMDRuntimeSettings* Settings = GetDefault<UOculusHMDRuntimeSettings>();
+	const UDEPRECATED_UOculusHMDRuntimeSettings* Settings = GetDefault<UDEPRECATED_UOculusHMDRuntimeSettings>();
 	return Settings->bSupportsDash ? EVisibility::Collapsed : EVisibility::Visible;
 }
 
diff --git a/Engine/Plugins/Runtime/Oculus/OculusVR/Source/OculusHMD/Private/OculusAssetManager.cpp b/Engine/Plugins/Runtime/Oculus/OculusVR/Source/OculusHMD/Private/OculusAssetManager.cpp
--- a/Engine/Plugins/Runtime/Oculus/OculusVR/Source/OculusHMD/Private/OculusAssetManager.cpp
+++ b/Engine/Plugins/Runtime/Oculus/OculusVR/Source/OculusHMD/Private/OculusAssetManager.cpp
@@ -14,6 +14,7 @@
 /* FOculusAssetDirectory
  *****************************************************************************/
 
+PRAGMA_DISABLE_DEPRECATION_WARNINGS
 FSoftObjectPath FOculusAssetDirectory::AssetListing[] =
 {
 	FString(TEXT("/OculusVR/Meshes/RiftHMD.RiftHMD")),
@@ -260,4 +261,5 @@ UPrimitiveComponent* FOculusAssetManager::CreateRenderComponent(const int32 Devi
 	OnLoadComplete.ExecuteIfBound(NewRenderComponent);
 	return NewRenderComponent;
 }
+PRAGMA_ENABLE_DEPRECATION_WARNINGS
 
diff --git a/Engine/Plugins/Runtime/Oculus/OculusVR/Source/OculusHMD/Private/OculusEventComponent.cpp b/Engine/Plugins/Runtime/Oculus/OculusVR/Source/OculusHMD/Private/OculusEventComponent.cpp
--- a/Engine/Plugins/Runtime/Oculus/OculusVR/Source/OculusHMD/Private/OculusEventComponent.cpp
+++ b/Engine/Plugins/Runtime/Oculus/OculusVR/Source/OculusHMD/Private/OculusEventComponent.cpp
@@ -3,14 +3,14 @@
 #include "OculusHMD.h"
 #include "OculusDelegates.h"
 
-void UOculusEventComponent::OnRegister()
+void UDEPRECATED_UOculusEventComponent::OnRegister()
 {
 	Super::OnRegister();
 
-	FOculusEventDelegates::OculusDisplayRefreshRateChanged.AddUObject(this, &UOculusEventComponent::OculusDisplayRefreshRateChanged_Handler);
+	FOculusEventDelegates::OculusDisplayRefreshRateChanged.AddUObject(this, &UDEPRECATED_UOculusEventComponent::OculusDisplayRefreshRateChanged_Handler);
 }
 
-void UOculusEventComponent::OnUnregister()
+void UDEPRECATED_UOculusEventComponent::OnUnregister()
 {
 	Super::OnUnregister();
 
diff --git a/Engine/Plugins/Runtime/Oculus/OculusVR/Source/OculusHMD/Private/OculusFunctionLibrary.cpp b/Engine/Plugins/Runtime/Oculus/OculusVR/Source/OculusHMD/Private/OculusFunctionLibrary.cpp
--- a/Engine/Plugins/Runtime/Oculus/OculusVR/Source/OculusHMD/Private/OculusFunctionLibrary.cpp
+++ b/Engine/Plugins/Runtime/Oculus/OculusVR/Source/OculusHMD/Private/OculusFunctionLibrary.cpp
@@ -7,16 +7,18 @@
 
 #define LOCTEXT_NAMESPACE "OculusFunctionLibrary"
 
+PRAGMA_DISABLE_DEPRECATION_WARNINGS
+
 //-------------------------------------------------------------------------------------------------
-// UOculusFunctionLibrary
+// UDEPRECATED_UOculusFunctionLibrary
 //-------------------------------------------------------------------------------------------------
 
-UOculusFunctionLibrary::UOculusFunctionLibrary(const FObjectInitializer& ObjectInitializer)
+UDEPRECATED_UOculusFunctionLibrary::UDEPRECATED_UOculusFunctionLibrary(const FObjectInitializer& ObjectInitializer)
 	: Super(ObjectInitializer)
 {
 }
 
-OculusHMD::FOculusHMD* UOculusFunctionLibrary::GetOculusHMD()
+OculusHMD::FOculusHMD* UDEPRECATED_UOculusFunctionLibrary::GetOculusHMD()
 {
 #if OCULUS_HMD_SUPPORTED_PLATFORMS
 	if (GEngine && GEngine->XRSystem.IsValid())
@@ -30,7 +32,7 @@ OculusHMD::FOculusHMD* UOculusFunctionLibrary::GetOculusHMD()
 	return nullptr;
 }
 
-void UOculusFunctionLibrary::GetPose(FRotator& DeviceRotation, FVector& DevicePosition, FVector& NeckPosition, bool bUseOrienationForPlayerCamera, bool bUsePositionForPlayerCamera, const FVector PositionScale)
+void UDEPRECATED_UOculusFunctionLibrary::GetPose(FRotator& DeviceRotation, FVector& DevicePosition, FVector& NeckPosition, bool bUseOrienationForPlayerCamera, bool bUsePositionForPlayerCamera, const FVector PositionScale)
 {
 #if OCULUS_HMD_SUPPORTED_PLATFORMS
 	OculusHMD::FOculusHMD* OculusHMD = GetOculusHMD();
@@ -54,7 +56,7 @@ void UOculusFunctionLibrary::GetPose(FRotator& DeviceRotation, FVector& DevicePo
 	}
 }
 
-void UOculusFunctionLibrary::SetBaseRotationAndBaseOffsetInMeters(FRotator Rotation, FVector BaseOffsetInMeters, EOrientPositionSelector::Type Options)
+void UDEPRECATED_UOculusFunctionLibrary::SetBaseRotationAndBaseOffsetInMeters(FRotator Rotation, FVector BaseOffsetInMeters, EOrientPositionSelector::Type Options)
 {
 #if OCULUS_HMD_SUPPORTED_PLATFORMS
 	OculusHMD::FOculusHMD* OculusHMD = GetOculusHMD();
@@ -72,7 +74,7 @@ void UOculusFunctionLibrary::SetBaseRotationAndBaseOffsetInMeters(FRotator Rotat
 #endif // OCULUS_HMD_SUPPORTED_PLATFORMS
 }
 
-void UOculusFunctionLibrary::GetBaseRotationAndBaseOffsetInMeters(FRotator& OutRotation, FVector& OutBaseOffsetInMeters)
+void UDEPRECATED_UOculusFunctionLibrary::GetBaseRotationAndBaseOffsetInMeters(FRotator& OutRotation, FVector& OutBaseOffsetInMeters)
 {
 #if OCULUS_HMD_SUPPORTED_PLATFORMS
 	OculusHMD::FOculusHMD* OculusHMD = GetOculusHMD();
@@ -89,7 +91,7 @@ void UOculusFunctionLibrary::GetBaseRotationAndBaseOffsetInMeters(FRotator& OutR
 #endif // OCULUS_HMD_SUPPORTED_PLATFORMS
 }
 
-void UOculusFunctionLibrary::GetRawSensorData(FVector& AngularAcceleration, FVector& LinearAcceleration, FVector& AngularVelocity, FVector& LinearVelocity, float& TimeInSeconds, ETrackedDeviceType DeviceType)
+void UDEPRECATED_UOculusFunctionLibrary::GetRawSensorData(FVector& AngularAcceleration, FVector& LinearAcceleration, FVector& AngularVelocity, FVector& LinearVelocity, float& TimeInSeconds, ETrackedDeviceType DeviceType)
 {
 #if OCULUS_HMD_SUPPORTED_PLATFORMS
 	OculusHMD::FOculusHMD* OculusHMD = GetOculusHMD();
@@ -108,7 +110,7 @@ void UOculusFunctionLibrary::GetRawSensorData(FVector& AngularAcceleration, FVec
 #endif // OCULUS_HMD_SUPPORTED_PLATFORMS
 }
 
-bool UOculusFunctionLibrary::IsDeviceTracked(ETrackedDeviceType DeviceType)
+bool UDEPRECATED_UOculusFunctionLibrary::IsDeviceTracked(ETrackedDeviceType DeviceType)
 {
 #if OCULUS_HMD_SUPPORTED_PLATFORMS
 	OculusHMD::FOculusHMD* OculusHMD = GetOculusHMD();
@@ -128,7 +130,7 @@ bool UOculusFunctionLibrary::IsDeviceTracked(ETrackedDeviceType DeviceType)
 	return false;
 }
 
-void UOculusFunctionLibrary::SetCPUAndGPULevels(int CPULevel, int GPULevel)
+void UDEPRECATED_UOculusFunctionLibrary::SetCPUAndGPULevels(int CPULevel, int GPULevel)
 {
 #if OCULUS_HMD_SUPPORTED_PLATFORMS
 	OculusHMD::FOculusHMD* OculusHMD = GetOculusHMD();
@@ -139,7 +141,7 @@ void UOculusFunctionLibrary::SetCPUAndGPULevels(int CPULevel, int GPULevel)
 #endif // OCULUS_HMD_SUPPORTED_PLATFORMS
 }
 
-bool UOculusFunctionLibrary::GetUserProfile(FHmdUserProfile& Profile)
+bool UDEPRECATED_UOculusFunctionLibrary::GetUserProfile(FHmdUserProfile& Profile)
 {
 #if OCULUS_HMD_SUPPORTED_PLATFORMS
 	OculusHMD::FOculusHMD* OculusHMD = GetOculusHMD();
@@ -161,7 +163,7 @@ bool UOculusFunctionLibrary::GetUserProfile(FHmdUserProfile& Profile)
 	return false;
 }
 
-void UOculusFunctionLibrary::SetBaseRotationAndPositionOffset(FRotator BaseRot, FVector PosOffset, EOrientPositionSelector::Type Options)
+void UDEPRECATED_UOculusFunctionLibrary::SetBaseRotationAndPositionOffset(FRotator BaseRot, FVector PosOffset, EOrientPositionSelector::Type Options)
 {
 #if OCULUS_HMD_SUPPORTED_PLATFORMS
 	OculusHMD::FOculusHMD* OculusHMD = GetOculusHMD();
@@ -175,7 +177,7 @@ void UOculusFunctionLibrary::SetBaseRotationAndPositionOffset(FRotator BaseRot,
 #endif // OCULUS_HMD_SUPPORTED_PLATFORMS
 }
 
-void UOculusFunctionLibrary::GetBaseRotationAndPositionOffset(FRotator& OutRot, FVector& OutPosOffset)
+void UDEPRECATED_UOculusFunctionLibrary::GetBaseRotationAndPositionOffset(FRotator& OutRot, FVector& OutPosOffset)
 {
 #if OCULUS_HMD_SUPPORTED_PLATFORMS
 	OculusHMD::FOculusHMD* OculusHMD = GetOculusHMD();
@@ -187,7 +189,7 @@ void UOculusFunctionLibrary::GetBaseRotationAndPositionOffset(FRotator& OutRot,
 #endif // OCULUS_HMD_SUPPORTED_PLATFORMS
 }
 
-void UOculusFunctionLibrary::AddLoadingSplashScreen(class UTexture2D* Texture, FVector TranslationInMeters, FRotator Rotation, FVector2D SizeInMeters, FRotator DeltaRotation, bool bClearBeforeAdd)
+void UDEPRECATED_UOculusFunctionLibrary::AddLoadingSplashScreen(class UTexture2D* Texture, FVector TranslationInMeters, FRotator Rotation, FVector2D SizeInMeters, FRotator DeltaRotation, bool bClearBeforeAdd)
 {
 #if OCULUS_HMD_SUPPORTED_PLATFORMS
 	OculusHMD::FOculusHMD* OculusHMD = GetOculusHMD();
@@ -212,7 +214,7 @@ void UOculusFunctionLibrary::AddLoadingSplashScreen(class UTexture2D* Texture, F
 #endif // OCULUS_HMD_SUPPORTED_PLATFORMS
 }
 
-void UOculusFunctionLibrary::ClearLoadingSplashScreens()
+void UDEPRECATED_UOculusFunctionLibrary::ClearLoadingSplashScreens()
 {
 #if OCULUS_HMD_SUPPORTED_PLATFORMS
 	OculusHMD::FOculusHMD* OculusHMD = GetOculusHMD();
@@ -227,7 +229,7 @@ void UOculusFunctionLibrary::ClearLoadingSplashScreens()
 #endif // OCULUS_HMD_SUPPORTED_PLATFORMS
 }
 
-bool UOculusFunctionLibrary::HasInputFocus()
+bool UDEPRECATED_UOculusFunctionLibrary::HasInputFocus()
 {
 #if OCULUS_HMD_SUPPORTED_PLATFORMS
 	const OculusHMD::FOculusHMD* const OculusHMD = GetOculusHMD();
@@ -243,7 +245,7 @@ bool UOculusFunctionLibrary::HasInputFocus()
 	return false;
 }
 
-bool UOculusFunctionLibrary::HasSystemOverlayPresent()
+bool UDEPRECATED_UOculusFunctionLibrary::HasSystemOverlayPresent()
 {
 #if OCULUS_HMD_SUPPORTED_PLATFORMS
 	const OculusHMD::FOculusHMD* const OculusHMD = GetOculusHMD();
@@ -259,7 +261,7 @@ bool UOculusFunctionLibrary::HasSystemOverlayPresent()
 	return false;
 }
 
-void UOculusFunctionLibrary::GetGPUUtilization(bool& IsGPUAvailable, float& GPUUtilization)
+void UDEPRECATED_UOculusFunctionLibrary::GetGPUUtilization(bool& IsGPUAvailable, float& GPUUtilization)
 {
 	IsGPUAvailable = false;
 	GPUUtilization = 0.0f;
@@ -278,7 +280,7 @@ void UOculusFunctionLibrary::GetGPUUtilization(bool& IsGPUAvailable, float& GPUU
 #endif // OCULUS_HMD_SUPPORTED_PLATFORMS
 }
 
-float UOculusFunctionLibrary::GetGPUFrameTime()
+float UDEPRECATED_UOculusFunctionLibrary::GetGPUFrameTime()
 {
 	float frameTime = 0;
 #if OCULUS_HMD_SUPPORTED_PLATFORMS
@@ -294,7 +296,7 @@ float UOculusFunctionLibrary::GetGPUFrameTime()
 	return 0.0f;
 }
 
-void UOculusFunctionLibrary::SetFixedFoveatedRenderingLevel(EFixedFoveatedRenderingLevel level, bool isDynamic)
+void UDEPRECATED_UOculusFunctionLibrary::SetFixedFoveatedRenderingLevel(EFixedFoveatedRenderingLevel level, bool isDynamic)
 {
 #if OCULUS_HMD_SUPPORTED_PLATFORMS
 	OculusHMD::FOculusHMD* OculusHMD = GetOculusHMD();
@@ -305,7 +307,7 @@ void UOculusFunctionLibrary::SetFixedFoveatedRenderingLevel(EFixedFoveatedRender
 #endif // OCULUS_HMD_SUPPORTED_PLATFORMS
 }
 
-EFixedFoveatedRenderingLevel UOculusFunctionLibrary::GetFixedFoveatedRenderingLevel()
+EFixedFoveatedRenderingLevel UDEPRECATED_UOculusFunctionLibrary::GetFixedFoveatedRenderingLevel()
 {
 #if OCULUS_HMD_SUPPORTED_PLATFORMS
 	OculusHMD::FOculusHMD* OculusHMD = GetOculusHMD();
@@ -321,7 +323,7 @@ EFixedFoveatedRenderingLevel UOculusFunctionLibrary::GetFixedFoveatedRenderingLe
 	return EFixedFoveatedRenderingLevel::FFR_Off;
 }
 
-FString UOculusFunctionLibrary::GetDeviceName()
+FString UDEPRECATED_UOculusFunctionLibrary::GetDeviceName()
 {
 #if OCULUS_HMD_SUPPORTED_PLATFORMS
 	OculusHMD::FOculusHMD* OculusHMD = GetOculusHMD();
@@ -337,7 +339,7 @@ FString UOculusFunctionLibrary::GetDeviceName()
 	return FString();
 }
 
-EOculusDeviceType UOculusFunctionLibrary::GetDeviceType()
+EOculusDeviceType UDEPRECATED_UOculusFunctionLibrary::GetDeviceType()
 {
 #if OCULUS_HMD_SUPPORTED_PLATFORMS
 	OculusHMD::FOculusHMD* OculusHMD = GetOculusHMD();
@@ -371,7 +373,7 @@ EOculusDeviceType UOculusFunctionLibrary::GetDeviceType()
 	return EOculusDeviceType::OculusUnknown;
 }
 
-TArray<float> UOculusFunctionLibrary::GetAvailableDisplayFrequencies()
+TArray<float> UDEPRECATED_UOculusFunctionLibrary::GetAvailableDisplayFrequencies()
 {
 #if OCULUS_HMD_SUPPORTED_PLATFORMS
 	OculusHMD::FOculusHMD* OculusHMD = GetOculusHMD();
@@ -390,7 +392,7 @@ TArray<float> UOculusFunctionLibrary::GetAvailableDisplayFrequencies()
 	return TArray<float>();
 }
 
-float UOculusFunctionLibrary::GetCurrentDisplayFrequency()
+float UDEPRECATED_UOculusFunctionLibrary::GetCurrentDisplayFrequency()
 {
 #if OCULUS_HMD_SUPPORTED_PLATFORMS
 	OculusHMD::FOculusHMD* OculusHMD = GetOculusHMD();
@@ -406,7 +408,7 @@ float UOculusFunctionLibrary::GetCurrentDisplayFrequency()
 	return 0.0f;
 }
 
-void UOculusFunctionLibrary::SetDisplayFrequency(float RequestedFrequency)
+void UDEPRECATED_UOculusFunctionLibrary::SetDisplayFrequency(float RequestedFrequency)
 {
 #if OCULUS_HMD_SUPPORTED_PLATFORMS
 	OculusHMD::FOculusHMD* OculusHMD = GetOculusHMD();
@@ -417,7 +419,7 @@ void UOculusFunctionLibrary::SetDisplayFrequency(float RequestedFrequency)
 #endif
 }
 
-void UOculusFunctionLibrary::EnablePositionTracking(bool bPositionTracking)
+void UDEPRECATED_UOculusFunctionLibrary::EnablePositionTracking(bool bPositionTracking)
 {
 #if OCULUS_HMD_SUPPORTED_PLATFORMS
 	OculusHMD::FOculusHMD* OculusHMD = GetOculusHMD();
@@ -429,7 +431,7 @@ void UOculusFunctionLibrary::EnablePositionTracking(bool bPositionTracking)
 }
 
 
-void UOculusFunctionLibrary::EnableOrientationTracking(bool bOrientationTracking)
+void UDEPRECATED_UOculusFunctionLibrary::EnableOrientationTracking(bool bOrientationTracking)
 {
 #if OCULUS_HMD_SUPPORTED_PLATFORMS
 	OculusHMD::FOculusHMD* OculusHMD = GetOculusHMD();
@@ -440,7 +442,7 @@ void UOculusFunctionLibrary::EnableOrientationTracking(bool bOrientationTracking
 #endif
 }
 
-void UOculusFunctionLibrary::SetColorScaleAndOffset(FLinearColor ColorScale, FLinearColor ColorOffset, bool bApplyToAllLayers)
+void UDEPRECATED_UOculusFunctionLibrary::SetColorScaleAndOffset(FLinearColor ColorScale, FLinearColor ColorOffset, bool bApplyToAllLayers)
 {
 #if OCULUS_HMD_SUPPORTED_PLATFORMS
 	OculusHMD::FOculusHMD* OculusHMD = GetOculusHMD();
@@ -451,7 +453,7 @@ void UOculusFunctionLibrary::SetColorScaleAndOffset(FLinearColor ColorScale, FLi
 #endif
 }
 
-class IStereoLayers* UOculusFunctionLibrary::GetStereoLayers()
+class IStereoLayers* UDEPRECATED_UOculusFunctionLibrary::GetStereoLayers()
 {
 #if OCULUS_HMD_SUPPORTED_PLATFORMS
 	OculusHMD::FOculusHMD* OculusHMD = GetOculusHMD();
@@ -479,7 +481,7 @@ static ovrpBoundaryType ToOvrpBoundaryType(EBoundaryType Source)
 }
 #endif // OCULUS_HMD_SUPPORTED_PLATFORMS
 
-bool UOculusFunctionLibrary::IsGuardianConfigured()
+bool UDEPRECATED_UOculusFunctionLibrary::IsGuardianConfigured()
 {
 #if OCULUS_HMD_SUPPORTED_PLATFORMS
 	OculusHMD::FOculusHMD* OculusHMD = GetOculusHMD();
@@ -492,7 +494,7 @@ bool UOculusFunctionLibrary::IsGuardianConfigured()
 	return false;
 }
 
-bool UOculusFunctionLibrary::IsGuardianDisplayed()
+bool UDEPRECATED_UOculusFunctionLibrary::IsGuardianDisplayed()
 {
 #if OCULUS_HMD_SUPPORTED_PLATFORMS
 	OculusHMD::FOculusHMD* OculusHMD = GetOculusHMD();
@@ -505,7 +507,7 @@ bool UOculusFunctionLibrary::IsGuardianDisplayed()
 	return false;
 }
 
-TArray<FVector> UOculusFunctionLibrary::GetGuardianPoints(EBoundaryType BoundaryType, bool UsePawnSpace /* = false */)
+TArray<FVector> UDEPRECATED_UOculusFunctionLibrary::GetGuardianPoints(EBoundaryType BoundaryType, bool UsePawnSpace /* = false */)
 {
 	TArray<FVector> BoundaryPointList;
 #if OCULUS_HMD_SUPPORTED_PLATFORMS
@@ -549,7 +551,7 @@ TArray<FVector> UOculusFunctionLibrary::GetGuardianPoints(EBoundaryType Boundary
 	return BoundaryPointList;
 }
 
-FVector UOculusFunctionLibrary::GetGuardianDimensions(EBoundaryType BoundaryType)
+FVector UDEPRECATED_UOculusFunctionLibrary::GetGuardianDimensions(EBoundaryType BoundaryType)
 {
 #if OCULUS_HMD_SUPPORTED_PLATFORMS
 	OculusHMD::FOculusHMD* OculusHMD = GetOculusHMD();
@@ -568,7 +570,7 @@ FVector UOculusFunctionLibrary::GetGuardianDimensions(EBoundaryType BoundaryType
 	return FVector::ZeroVector;
 }
 
-FTransform UOculusFunctionLibrary::GetPlayAreaTransform()
+FTransform UDEPRECATED_UOculusFunctionLibrary::GetPlayAreaTransform()
 {
 #if OCULUS_HMD_SUPPORTED_PLATFORMS
 	OculusHMD::FOculusHMD* OculusHMD = GetOculusHMD();
@@ -602,7 +604,7 @@ FTransform UOculusFunctionLibrary::GetPlayAreaTransform()
 	return FTransform();
 }
 
-FGuardianTestResult UOculusFunctionLibrary::GetPointGuardianIntersection(const FVector Point, EBoundaryType BoundaryType)
+FGuardianTestResult UDEPRECATED_UOculusFunctionLibrary::GetPointGuardianIntersection(const FVector Point, EBoundaryType BoundaryType)
 {
 	FGuardianTestResult InteractionInfo;
 	memset(&InteractionInfo, 0, sizeof(FGuardianTestResult));
@@ -629,7 +631,7 @@ FGuardianTestResult UOculusFunctionLibrary::GetPointGuardianIntersection(const F
 	return InteractionInfo;
 }
 
-FGuardianTestResult UOculusFunctionLibrary::GetNodeGuardianIntersection(ETrackedDeviceType DeviceType, EBoundaryType BoundaryType)
+FGuardianTestResult UDEPRECATED_UOculusFunctionLibrary::GetNodeGuardianIntersection(ETrackedDeviceType DeviceType, EBoundaryType BoundaryType)
 {
 	FGuardianTestResult InteractionInfo;
 	memset(&InteractionInfo, 0, sizeof(FGuardianTestResult));
@@ -656,7 +658,7 @@ FGuardianTestResult UOculusFunctionLibrary::GetNodeGuardianIntersection(ETracked
 	return InteractionInfo;
 }
 
-void UOculusFunctionLibrary::SetGuardianVisibility(bool GuardianVisible)
+void UDEPRECATED_UOculusFunctionLibrary::SetGuardianVisibility(bool GuardianVisible)
 {
 #if OCULUS_HMD_SUPPORTED_PLATFORMS
 	OculusHMD::FOculusHMD* OculusHMD = GetOculusHMD();
@@ -667,7 +669,7 @@ void UOculusFunctionLibrary::SetGuardianVisibility(bool GuardianVisible)
 #endif
 }
 
-bool UOculusFunctionLibrary::GetSystemHmd3DofModeEnabled()
+bool UDEPRECATED_UOculusFunctionLibrary::GetSystemHmd3DofModeEnabled()
 {
 #if OCULUS_HMD_SUPPORTED_PLATFORMS
 	OculusHMD::FOculusHMD* OculusHMD = GetOculusHMD();
@@ -680,7 +682,7 @@ bool UOculusFunctionLibrary::GetSystemHmd3DofModeEnabled()
 	return false;
 }
 
-EOculusColorSpace UOculusFunctionLibrary::GetHmdColorDesc()
+EOculusColorSpace UDEPRECATED_UOculusFunctionLibrary::GetHmdColorDesc()
 {
 #if OCULUS_HMD_SUPPORTED_PLATFORMS
 	OculusHMD::FOculusHMD* OculusHMD = GetOculusHMD();
@@ -696,7 +698,7 @@ EOculusColorSpace UOculusFunctionLibrary::GetHmdColorDesc()
 	return EOculusColorSpace::Unknown;
 }
 
-void UOculusFunctionLibrary::SetClientColorDesc(EOculusColorSpace ColorSpace)
+void UDEPRECATED_UOculusFunctionLibrary::SetClientColorDesc(EOculusColorSpace ColorSpace)
 {
 #if OCULUS_HMD_SUPPORTED_PLATFORMS
 	OculusHMD::FOculusHMD* OculusHMD = GetOculusHMD();
@@ -714,4 +716,6 @@ void UOculusFunctionLibrary::SetClientColorDesc(EOculusColorSpace ColorSpace)
 #endif
 }
 
+PRAGMA_ENABLE_DEPRECATION_WARNINGS
+
 #undef LOCTEXT_NAMESPACE
\ No newline at end of file
diff --git a/Engine/Plugins/Runtime/Oculus/OculusVR/Source/OculusHMD/Private/OculusHMD.cpp b/Engine/Plugins/Runtime/Oculus/OculusVR/Source/OculusHMD/Private/OculusHMD.cpp
--- a/Engine/Plugins/Runtime/Oculus/OculusVR/Source/OculusHMD/Private/OculusHMD.cpp
+++ b/Engine/Plugins/Runtime/Oculus/OculusVR/Source/OculusHMD/Private/OculusHMD.cpp
@@ -51,6 +51,8 @@
 
 #if OCULUS_HMD_SUPPORTED_PLATFORMS
 
+PRAGMA_DISABLE_DEPRECATION_WARNINGS
+
 static TAutoConsoleVariable<int32> CVarOculusEnableSubsampledLayout(
 	TEXT("r.Mobile.Oculus.EnableSubsampled"),
 	0,
@@ -3691,7 +3693,7 @@ namespace OculusHMD
 
 	void FOculusHMD::LoadFromSettings()
 	{
-		UOculusHMDRuntimeSettings* HMDSettings = GetMutableDefault<UOculusHMDRuntimeSettings>();
+		UDEPRECATED_UOculusHMDRuntimeSettings* HMDSettings = GetMutableDefault<UDEPRECATED_UOculusHMDRuntimeSettings>();
 		check(HMDSettings);
 
 		Settings->Flags.bSupportsDash = HMDSettings->bSupportsDash;
@@ -3705,7 +3707,6 @@ namespace OculusHMD
 		Settings->bLateLatching = HMDSettings->bLateLatching;
 		Settings->bPhaseSync = HMDSettings->bPhaseSync;
 
-
 		// Set FFR level and dynamic from rendering settings
 		UEnum* FFREnum = StaticEnum<EFixedFoveatedRenderingLevel>();
 		int32 VRSValue = IConsoleManager::Get().FindConsoleVariable(TEXT("vr.VRS.HMDFixedFoveationLevel"))->GetInt();
@@ -3720,4 +3721,6 @@ namespace OculusHMD
 
 } // namespace OculusHMD
 
+PRAGMA_ENABLE_DEPRECATION_WARNINGS
+
 #endif //OCULUS_HMD_SUPPORTED_PLATFORMS
diff --git a/Engine/Plugins/Runtime/Oculus/OculusVR/Source/OculusHMD/Private/OculusHMD.h b/Engine/Plugins/Runtime/Oculus/OculusVR/Source/OculusHMD/Private/OculusHMD.h
--- a/Engine/Plugins/Runtime/Oculus/OculusVR/Source/OculusHMD/Private/OculusHMD.h
+++ b/Engine/Plugins/Runtime/Oculus/OculusVR/Source/OculusHMD/Private/OculusHMD.h
@@ -28,6 +28,7 @@
 #include "XRThreadUtils.h"
 #include "ProceduralMeshComponent.h"
 
+PRAGMA_DISABLE_DEPRECATION_WARNINGS
 
 namespace OculusHMD
 {
@@ -67,7 +68,7 @@ enum FRecenterTypes
 
 class FOculusHMD : public FHeadMountedDisplayBase, public FXRRenderTargetManager, public IStereoLayers, public FHMDSceneViewExtension, public FOculusAssetManager
 {
-	friend class UOculusFunctionLibrary;
+	friend class UDEPRECATED_UOculusFunctionLibrary;
 	friend FOculusHMDModule;
 	friend class FSplash;
 	friend class FConsoleCommands;
@@ -468,4 +469,6 @@ typedef TSharedPtr< FOculusHMD, ESPMode::ThreadSafe > FOculusHMDPtr;
 
 } // namespace OculusHMD
 
+PRAGMA_ENABLE_DEPRECATION_WARNINGS
+
 #endif //OCULUS_HMD_SUPPORTED_PLATFORMS
diff --git a/Engine/Plugins/Runtime/Oculus/OculusVR/Source/OculusHMD/Private/OculusHMDModule.cpp b/Engine/Plugins/Runtime/Oculus/OculusVR/Source/OculusHMD/Private/OculusHMDModule.cpp
--- a/Engine/Plugins/Runtime/Oculus/OculusVR/Source/OculusHMD/Private/OculusHMDModule.cpp
+++ b/Engine/Plugins/Runtime/Oculus/OculusVR/Source/OculusHMD/Private/OculusHMDModule.cpp
@@ -21,6 +21,7 @@
 
 OculusPluginWrapper FOculusHMDModule::PluginWrapper;
 
+PRAGMA_DISABLE_DEPRECATION_WARNINGS
 FOculusHMDModule::FOculusHMDModule()
 {
 #if OCULUS_HMD_SUPPORTED_PLATFORMS
@@ -174,7 +175,7 @@ bool FOculusHMDModule::PreInit()
 bool FOculusHMDModule::IsHMDConnected()
 {
 #if OCULUS_HMD_SUPPORTED_PLATFORMS
-	UOculusHMDRuntimeSettings* HMDSettings = GetMutableDefault<UOculusHMDRuntimeSettings>();
+	UDEPRECATED_UOculusHMDRuntimeSettings* HMDSettings = GetMutableDefault<UDEPRECATED_UOculusHMDRuntimeSettings>();
 	if (FApp::CanEverRender() && OculusHMD::IsOculusHMDConnected() && HMDSettings->XrApi != EOculusXrApi::NativeOpenXR)
 	{
 		return true;
@@ -412,4 +413,6 @@ void FOculusHMDModule::SetGraphicsAdapterLuid(uint64 InLuid)
 }
 #endif // OCULUS_HMD_SUPPORTED_PLATFORMS
 
+PRAGMA_ENABLE_DEPRECATION_WARNINGS
+
 IMPLEMENT_MODULE(FOculusHMDModule, OculusHMD)
diff --git a/Engine/Plugins/Runtime/Oculus/OculusVR/Source/OculusHMD/Private/OculusHMDModule.h b/Engine/Plugins/Runtime/Oculus/OculusVR/Source/OculusHMD/Private/OculusHMDModule.h
--- a/Engine/Plugins/Runtime/Oculus/OculusVR/Source/OculusHMD/Private/OculusHMDModule.h
+++ b/Engine/Plugins/Runtime/Oculus/OculusVR/Source/OculusHMD/Private/OculusHMDModule.h
@@ -11,6 +11,8 @@
 // FOculusHMDModule
 //-------------------------------------------------------------------------------------------------
 
+PRAGMA_DISABLE_DEPRECATION_WARNINGS
+
 class FOculusHMDModule : public IOculusHMDModule
 {
 public:
@@ -42,42 +44,42 @@ public:
 	// IOculusHMDModule
 	virtual void GetPose(FRotator& DeviceRotation, FVector& DevicePosition, FVector& NeckPosition, bool bUseOrienationForPlayerCamera = false, bool bUsePositionForPlayerCamera = false, const FVector PositionScale = FVector::ZeroVector) override
 	{
-		UOculusFunctionLibrary::GetPose(DeviceRotation, DevicePosition, NeckPosition, bUseOrienationForPlayerCamera, bUsePositionForPlayerCamera, PositionScale);
+		UDEPRECATED_UOculusFunctionLibrary::GetPose(DeviceRotation, DevicePosition, NeckPosition, bUseOrienationForPlayerCamera, bUsePositionForPlayerCamera, PositionScale);
 	}
 
 	virtual void GetRawSensorData(FVector& AngularAcceleration, FVector& LinearAcceleration, FVector& AngularVelocity, FVector& LinearVelocity, float& TimeInSeconds) override
 	{
-		UOculusFunctionLibrary::GetRawSensorData(AngularAcceleration, LinearAcceleration, AngularVelocity, LinearVelocity, TimeInSeconds, ETrackedDeviceType::HMD);
+		UDEPRECATED_UOculusFunctionLibrary::GetRawSensorData(AngularAcceleration, LinearAcceleration, AngularVelocity, LinearVelocity, TimeInSeconds, ETrackedDeviceType::HMD);
 	}
 
 	virtual bool GetUserProfile(struct FHmdUserProfile& Profile) override
 	{
-		return UOculusFunctionLibrary::GetUserProfile(Profile);
+		return UDEPRECATED_UOculusFunctionLibrary::GetUserProfile(Profile);
 	}
 
 	virtual void SetBaseRotationAndBaseOffsetInMeters(FRotator Rotation, FVector BaseOffsetInMeters, EOrientPositionSelector::Type Options) override
 	{
-		UOculusFunctionLibrary::SetBaseRotationAndBaseOffsetInMeters(Rotation, BaseOffsetInMeters, Options);
+		UDEPRECATED_UOculusFunctionLibrary::SetBaseRotationAndBaseOffsetInMeters(Rotation, BaseOffsetInMeters, Options);
 	}
 
 	virtual void GetBaseRotationAndBaseOffsetInMeters(FRotator& OutRotation, FVector& OutBaseOffsetInMeters) override
 	{
-		UOculusFunctionLibrary::GetBaseRotationAndBaseOffsetInMeters(OutRotation, OutBaseOffsetInMeters);
+		UDEPRECATED_UOculusFunctionLibrary::GetBaseRotationAndBaseOffsetInMeters(OutRotation, OutBaseOffsetInMeters);
 	}
 
 	virtual void SetBaseRotationAndPositionOffset(FRotator BaseRot, FVector PosOffset, EOrientPositionSelector::Type Options) override
 	{
-		UOculusFunctionLibrary::SetBaseRotationAndPositionOffset(BaseRot, PosOffset, Options);
+		UDEPRECATED_UOculusFunctionLibrary::SetBaseRotationAndPositionOffset(BaseRot, PosOffset, Options);
 	}
 
 	virtual void GetBaseRotationAndPositionOffset(FRotator& OutRot, FVector& OutPosOffset) override
 	{
-		UOculusFunctionLibrary::GetBaseRotationAndPositionOffset(OutRot, OutPosOffset);
+		UDEPRECATED_UOculusFunctionLibrary::GetBaseRotationAndPositionOffset(OutRot, OutPosOffset);
 	}
 
 	virtual class IStereoLayers* GetStereoLayers() override
 	{
-		return UOculusFunctionLibrary::GetStereoLayers();
+		return UDEPRECATED_UOculusFunctionLibrary::GetStereoLayers();
 	}
 
 	bool IsOVRPluginAvailable() const
@@ -108,3 +110,5 @@ protected:
 
 #endif // OCULUS_HMD_SUPPORTED_PLATFORMS
 };
+
+PRAGMA_ENABLE_DEPRECATION_WARNINGS
diff --git a/Engine/Plugins/Runtime/Oculus/OculusVR/Source/OculusHMD/Private/OculusHMDPrivate.h b/Engine/Plugins/Runtime/Oculus/OculusVR/Source/OculusHMD/Private/OculusHMDPrivate.h
--- a/Engine/Plugins/Runtime/Oculus/OculusVR/Source/OculusHMD/Private/OculusHMDPrivate.h
+++ b/Engine/Plugins/Runtime/Oculus/OculusVR/Source/OculusHMD/Private/OculusHMDPrivate.h
@@ -33,6 +33,8 @@
 // Utility functions
 //-------------------------------------------------------------------------------------------------
 
+PRAGMA_DISABLE_DEPRECATION_WARNINGS
+
 namespace OculusHMD
 {
 #if OCULUS_HMD_SUPPORTED_PLATFORMS
@@ -254,3 +256,6 @@ namespace OculusHMD
 #endif // OCULUS_HMD_SUPPORTED_PLATFORMS
 
 } // namespace OculusHMD
+
+PRAGMA_ENABLE_DEPRECATION_WARNINGS
+
diff --git a/Engine/Plugins/Runtime/Oculus/OculusVR/Source/OculusHMD/Private/OculusHMDRuntimeSettings.cpp b/Engine/Plugins/Runtime/Oculus/OculusVR/Source/OculusHMD/Private/OculusHMDRuntimeSettings.cpp
--- a/Engine/Plugins/Runtime/Oculus/OculusVR/Source/OculusHMD/Private/OculusHMDRuntimeSettings.cpp
+++ b/Engine/Plugins/Runtime/Oculus/OculusVR/Source/OculusHMD/Private/OculusHMDRuntimeSettings.cpp
@@ -3,11 +3,13 @@
 #include "OculusHMDRuntimeSettings.h"
 
 //////////////////////////////////////////////////////////////////////////
-// UOculusHMDRuntimeSettings
+// UDEPRECATED_UOculusHMDRuntimeSettings
 
 #include "OculusHMD_Settings.h"
 
-UOculusHMDRuntimeSettings::UOculusHMDRuntimeSettings(const FObjectInitializer& ObjectInitializer)
+PRAGMA_DISABLE_DEPRECATION_WARNINGS
+
+UDEPRECATED_UOculusHMDRuntimeSettings::UDEPRECATED_UOculusHMDRuntimeSettings(const FObjectInitializer& ObjectInitializer)
 : Super(ObjectInitializer)
 , bAutoEnabled(true)
 {
@@ -52,14 +54,14 @@ UOculusHMDRuntimeSettings::UOculusHMDRuntimeSettings(const FObjectInitializer& O
 }
 
 #if WITH_EDITOR
-bool UOculusHMDRuntimeSettings::CanEditChange(const FProperty* InProperty) const
+bool UDEPRECATED_UOculusHMDRuntimeSettings::CanEditChange(const FProperty* InProperty) const
 {
 	bool bIsEditable = Super::CanEditChange(InProperty);
 	if (bIsEditable && InProperty)
 	{
 		const FName PropertyName = InProperty->GetFName();
 
-		if (PropertyName == GET_MEMBER_NAME_CHECKED(UOculusHMDRuntimeSettings, bFocusAware))
+		if (PropertyName == GET_MEMBER_NAME_CHECKED(UDEPRECATED_UOculusHMDRuntimeSettings, bFocusAware))
 		{
 			bIsEditable = false;
 		}
@@ -69,7 +71,7 @@ bool UOculusHMDRuntimeSettings::CanEditChange(const FProperty* InProperty) const
 }
 #endif // WITH_EDITOR
 
-void UOculusHMDRuntimeSettings::LoadFromIni()
+void UDEPRECATED_UOculusHMDRuntimeSettings::LoadFromIni()
 {
 	const TCHAR* OculusSettings = TEXT("Oculus.Settings");
 	bool v;
@@ -94,4 +96,6 @@ void UOculusHMDRuntimeSettings::LoadFromIni()
 	{
 		bCompositesDepth = v;
 	}
-}
\ No newline at end of file
+}
+
+PRAGMA_ENABLE_DEPRECATION_WARNINGS
diff --git a/Engine/Plugins/Runtime/Oculus/OculusVR/Source/OculusHMD/Private/OculusHMD_GameFrame.h b/Engine/Plugins/Runtime/Oculus/OculusVR/Source/OculusHMD/Private/OculusHMD_GameFrame.h
--- a/Engine/Plugins/Runtime/Oculus/OculusVR/Source/OculusHMD/Private/OculusHMD_GameFrame.h
+++ b/Engine/Plugins/Runtime/Oculus/OculusVR/Source/OculusHMD/Private/OculusHMD_GameFrame.h
@@ -7,6 +7,8 @@
 #include "OculusHMD_Settings.h"
 #include "ShowFlags.h"
 
+PRAGMA_DISABLE_DEPRECATION_WARNINGS
+
 namespace OculusHMD
 {
 
@@ -57,4 +59,6 @@ typedef TSharedPtr<FGameFrame, ESPMode::ThreadSafe> FGameFramePtr;
 
 } // namespace OculusHMD
 
+PRAGMA_ENABLE_DEPRECATION_WARNINGS
+
 #endif //OCULUS_HMD_SUPPORTED_PLATFORMS
diff --git a/Engine/Plugins/Runtime/Oculus/OculusVR/Source/OculusHMD/Private/OculusHMD_Settings.cpp b/Engine/Plugins/Runtime/Oculus/OculusVR/Source/OculusHMD/Private/OculusHMD_Settings.cpp
--- a/Engine/Plugins/Runtime/Oculus/OculusVR/Source/OculusHMD/Private/OculusHMD_Settings.cpp
+++ b/Engine/Plugins/Runtime/Oculus/OculusVR/Source/OculusHMD/Private/OculusHMD_Settings.cpp
@@ -4,6 +4,8 @@
 
 #if OCULUS_HMD_SUPPORTED_PLATFORMS
 
+PRAGMA_DISABLE_DEPRECATION_WARNINGS
+
 namespace OculusHMD
 {
 
@@ -92,4 +94,6 @@ void FSettings::SetPixelDensityMax(float NewPixelDensityMax)
 
 } // namespace OculusHMD
 
+PRAGMA_ENABLE_DEPRECATION_WARNINGS
+
 #endif //OCULUS_HMD_SUPPORTED_PLATFORMS
diff --git a/Engine/Plugins/Runtime/Oculus/OculusVR/Source/OculusHMD/Private/OculusHMD_Settings.h b/Engine/Plugins/Runtime/Oculus/OculusVR/Source/OculusHMD/Private/OculusHMD_Settings.h
--- a/Engine/Plugins/Runtime/Oculus/OculusVR/Source/OculusHMD/Private/OculusHMD_Settings.h
+++ b/Engine/Plugins/Runtime/Oculus/OculusVR/Source/OculusHMD/Private/OculusHMD_Settings.h
@@ -5,6 +5,8 @@
 
 #if OCULUS_HMD_SUPPORTED_PLATFORMS
 
+PRAGMA_DISABLE_DEPRECATION_WARNINGS
+
 namespace OculusHMD
 {
 
@@ -129,4 +131,6 @@ typedef TSharedPtr<FSettings, ESPMode::ThreadSafe> FSettingsPtr;
 
 } // namespace OculusHMD
 
+PRAGMA_ENABLE_DEPRECATION_WARNINGS
+
 #endif //OCULUS_HMD_SUPPORTED_PLATFORMS
diff --git a/Engine/Plugins/Runtime/Oculus/OculusVR/Source/OculusHMD/Private/OculusHMD_Splash.cpp b/Engine/Plugins/Runtime/Oculus/OculusVR/Source/OculusHMD/Private/OculusHMD_Splash.cpp
--- a/Engine/Plugins/Runtime/Oculus/OculusVR/Source/OculusHMD/Private/OculusHMD_Splash.cpp
+++ b/Engine/Plugins/Runtime/Oculus/OculusVR/Source/OculusHMD/Private/OculusHMD_Splash.cpp
@@ -14,6 +14,8 @@
 #include "OculusHMDTypes.h"
 #endif
 
+PRAGMA_DISABLE_DEPRECATION_WARNINGS
+
 namespace OculusHMD
 {
 
@@ -85,7 +87,7 @@ void FSplash::Tick_RenderThread(float DeltaTime)
 
 void FSplash::LoadSettings()
 {
-	UOculusHMDRuntimeSettings* HMDSettings = GetMutableDefault<UOculusHMDRuntimeSettings>();
+	UDEPRECATED_UOculusHMDRuntimeSettings* HMDSettings = GetMutableDefault<UDEPRECATED_UOculusHMDRuntimeSettings>();
 	check(HMDSettings);
 	ClearSplashes();
 	for (const FOculusSplashDesc& SplashDesc : HMDSettings->SplashDescs)
@@ -165,7 +167,7 @@ void FSplash::Startup()
 #if WITH_EDITOR
 		PieBeginDelegateHandle = FEditorDelegates::BeginPIE.AddRaw(this, &FSplash::OnPieBegin);
 #else
-		UOculusHMDRuntimeSettings* HMDSettings = GetMutableDefault<UOculusHMDRuntimeSettings>();
+		UDEPRECATED_UOculusHMDRuntimeSettings* HMDSettings = GetMutableDefault<UDEPRECATED_UOculusHMDRuntimeSettings>();
 		check(HMDSettings);
 		if (HMDSettings->bAutoEnabled)
 		{
@@ -670,4 +672,6 @@ void FSplash::UnloadTexture(FSplashLayer& InSplashLayer)
 
 } // namespace OculusHMD
 
+PRAGMA_ENABLE_DEPRECATION_WARNINGS
+
 #endif // OCULUS_HMD_SUPPORTED_PLATFORMS
diff --git a/Engine/Plugins/Runtime/Oculus/OculusVR/Source/OculusHMD/Public/IOculusHMDModule.h b/Engine/Plugins/Runtime/Oculus/OculusVR/Source/OculusHMD/Public/IOculusHMDModule.h
--- a/Engine/Plugins/Runtime/Oculus/OculusVR/Source/OculusHMD/Public/IOculusHMDModule.h
+++ b/Engine/Plugins/Runtime/Oculus/OculusVR/Source/OculusHMD/Public/IOculusHMDModule.h
@@ -15,6 +15,7 @@
  * The public interface to this module.  In most cases, this interface is only public to sibling modules
  * within this plugin.
  */
+class UE_DEPRECATED(5.1, "OculusVR plugin is deprecated; please use the built-in OpenXR plugin or OculusXR plugin from the Marketplace.") IOculusHMDModule;
 class IOculusHMDModule : public IHeadMountedDisplayModule
 {
 public:
diff --git a/Engine/Plugins/Runtime/Oculus/OculusVR/Source/OculusHMD/Public/OculusAssetDirectory.h b/Engine/Plugins/Runtime/Oculus/OculusVR/Source/OculusHMD/Public/OculusAssetDirectory.h
--- a/Engine/Plugins/Runtime/Oculus/OculusVR/Source/OculusHMD/Public/OculusAssetDirectory.h
+++ b/Engine/Plugins/Runtime/Oculus/OculusVR/Source/OculusHMD/Public/OculusAssetDirectory.h
@@ -5,6 +5,8 @@
 #include "CoreMinimal.h"
 #include "UObject/SoftObjectPath.h"
 
+class UE_DEPRECATED(5.1, "OculusVR plugin is deprecated; please use the built-in OpenXR plugin or OculusXR plugin from the Marketplace.") FOculusAssetDirectory;
+
 class FOculusAssetDirectory
 {
 public: 
diff --git a/Engine/Plugins/Runtime/Oculus/OculusVR/Source/OculusHMD/Public/OculusEventComponent.h b/Engine/Plugins/Runtime/Oculus/OculusVR/Source/OculusHMD/Public/OculusEventComponent.h
--- a/Engine/Plugins/Runtime/Oculus/OculusVR/Source/OculusHMD/Public/OculusEventComponent.h
+++ b/Engine/Plugins/Runtime/Oculus/OculusVR/Source/OculusHMD/Public/OculusEventComponent.h
@@ -8,8 +8,8 @@
 #include "Components/ActorComponent.h"
 #include "OculusEventComponent.generated.h"
 
-UCLASS(Blueprintable, meta = (BlueprintSpawnableComponent), ClassGroup = OculusHMD)
-class OCULUSHMD_API UOculusEventComponent : public UActorComponent
+UCLASS(Blueprintable, meta = (BlueprintSpawnableComponent, DeprecationMessage = "OculusVR plugin is deprecated; please use the built-in OpenXR plugin or OculusXR plugin from the Marketplace."), ClassGroup = OculusHMD, deprecated)
+class OCULUSHMD_API UDEPRECATED_UOculusEventComponent : public UActorComponent
 {
 	GENERATED_BODY()
 
diff --git a/Engine/Plugins/Runtime/Oculus/OculusVR/Source/OculusHMD/Public/OculusFunctionLibrary.h b/Engine/Plugins/Runtime/Oculus/OculusVR/Source/OculusHMD/Public/OculusFunctionLibrary.h
--- a/Engine/Plugins/Runtime/Oculus/OculusVR/Source/OculusHMD/Public/OculusFunctionLibrary.h
+++ b/Engine/Plugins/Runtime/Oculus/OculusVR/Source/OculusHMD/Public/OculusFunctionLibrary.h
@@ -13,9 +13,11 @@ namespace OculusHMD
 	class FOculusHMD;
 }
 
+PRAGMA_DISABLE_DEPRECATION_WARNINGS
+
 /* Tracked device types corresponding to ovrTrackedDeviceType enum*/
 UENUM(BlueprintType)
-enum class ETrackedDeviceType : uint8
+enum class UE_DEPRECATED(5.1, "OculusVR plugin is deprecated; please use the built-in OpenXR plugin or OculusXR plugin from the Marketplace.") ETrackedDeviceType : uint8
 {
 	None UMETA(DisplayName = "No Devices"),
 	HMD	UMETA(DisplayName = "HMD"),
@@ -26,15 +28,17 @@ enum class ETrackedDeviceType : uint8
 	All	UMETA(DisplayName = "All Devices")
 };
 
-USTRUCT(BlueprintType, meta = (DisplayName = "HMD User Profile Data Field"))
+USTRUCT(BlueprintType, meta = (DisplayName = "HMD User Profile Data Field", Deprecated = "5.1"))
 struct FHmdUserProfileField
 {
 	GENERATED_USTRUCT_BODY()
 
-	UPROPERTY(BlueprintReadWrite, Category = "Input|HeadMountedDisplay")
+	UE_DEPRECATED(5.1, "OculusVR plugin is deprecated; please use the built-in OpenXR plugin or OculusXR plugin from the Marketplace.")
+	UPROPERTY(BlueprintReadWrite, Category = "Input|HeadMountedDisplay", meta = (DeprecatedProperty))
 	FString FieldName;
 
-	UPROPERTY(BlueprintReadWrite, Category = "Input|HeadMountedDisplay")
+	UE_DEPRECATED(5.1, "OculusVR plugin is deprecated; please use the built-in OpenXR plugin or OculusXR plugin from the Marketplace.")
+	UPROPERTY(BlueprintReadWrite, Category = "Input|HeadMountedDisplay", meta = (DeprecatedProperty))
 	FString FieldValue;
 
 	FHmdUserProfileField() {}
@@ -42,36 +46,43 @@ struct FHmdUserProfileField
 		FieldName(Name), FieldValue(Value) {}
 };
 
-USTRUCT(BlueprintType, meta = (DisplayName = "HMD User Profile Data"))
+USTRUCT(BlueprintType, meta = (DisplayName = "HMD User Profile Data", Deprecated = "5.1"))
 struct FHmdUserProfile
 {
 	GENERATED_USTRUCT_BODY()
 
 	/** Name of the user's profile. */
-	UPROPERTY(BlueprintReadWrite, Category = "Input|HeadMountedDisplay")
+	UE_DEPRECATED(5.1, "OculusVR plugin is deprecated; please use the built-in OpenXR plugin or OculusXR plugin from the Marketplace.")
+	UPROPERTY(BlueprintReadWrite, Category = "Input|HeadMountedDisplay", meta = (DeprecatedProperty))
 	FString Name;
 
 	/** Gender of the user ("male", "female", etc). */
-	UPROPERTY(BlueprintReadWrite, Category = "Input|HeadMountedDisplay")
+	UE_DEPRECATED(5.1, "OculusVR plugin is deprecated; please use the built-in OpenXR plugin or OculusXR plugin from the Marketplace.")
+	UPROPERTY(BlueprintReadWrite, Category = "Input|HeadMountedDisplay", meta = (DeprecatedProperty))
 	FString Gender;
 
 	/** Height of the player, in meters */
-	UPROPERTY(BlueprintReadWrite, Category = "Input|HeadMountedDisplay")
+	UE_DEPRECATED(5.1, "OculusVR plugin is deprecated; please use the built-in OpenXR plugin or OculusXR plugin from the Marketplace.")
+	UPROPERTY(BlueprintReadWrite, Category = "Input|HeadMountedDisplay", meta = (DeprecatedProperty))
 	float PlayerHeight;
 
 	/** Height of the player, in meters */
-	UPROPERTY(BlueprintReadWrite, Category = "Input|HeadMountedDisplay")
+	UE_DEPRECATED(5.1, "OculusVR plugin is deprecated; please use the built-in OpenXR plugin or OculusXR plugin from the Marketplace.")
+	UPROPERTY(BlueprintReadWrite, Category = "Input|HeadMountedDisplay", meta = (DeprecatedProperty))
 	float EyeHeight;
 
 	/** Interpupillary distance of the player, in meters */
-	UPROPERTY(BlueprintReadWrite, Category = "Input|HeadMountedDisplay")
+	UE_DEPRECATED(5.1, "OculusVR plugin is deprecated; please use the built-in OpenXR plugin or OculusXR plugin from the Marketplace.")
+	UPROPERTY(BlueprintReadWrite, Category = "Input|HeadMountedDisplay", meta = (DeprecatedProperty))
 	float IPD;
 
 	/** Neck-to-eye distance, in meters. X - horizontal, Y - vertical. */
-	UPROPERTY(BlueprintReadWrite, Category = "Input|HeadMountedDisplay")
+	UE_DEPRECATED(5.1, "OculusVR plugin is deprecated; please use the built-in OpenXR plugin or OculusXR plugin from the Marketplace.")
+	UPROPERTY(BlueprintReadWrite, Category = "Input|HeadMountedDisplay", meta = (DeprecatedProperty))
 	FVector2D NeckToEyeDistance;
 
-	UPROPERTY(BlueprintReadWrite, Category = "Input|HeadMountedDisplay")
+	UE_DEPRECATED(5.1, "OculusVR plugin is deprecated; please use the built-in OpenXR plugin or OculusXR plugin from the Marketplace.")
+	UPROPERTY(BlueprintReadWrite, Category = "Input|HeadMountedDisplay", meta = (DeprecatedProperty))
 	TArray<FHmdUserProfileField> ExtraFields;
 
 	FHmdUserProfile() :
@@ -79,7 +90,7 @@ struct FHmdUserProfile
 };
 
 UENUM(BlueprintType)
-enum class EFixedFoveatedRenderingLevel : uint8
+enum class UE_DEPRECATED(5.1, "OculusVR plugin is deprecated; please use the built-in OpenXR plugin or OculusXR plugin from the Marketplace.") EFixedFoveatedRenderingLevel : uint8
 {
 	FFR_Off = 0,
 	FFR_Low = 1,
@@ -91,14 +102,14 @@ enum class EFixedFoveatedRenderingLevel : uint8
 
 /* Guardian boundary types*/
 UENUM(BlueprintType)
-enum class EBoundaryType : uint8
+enum class UE_DEPRECATED(5.1, "OculusVR plugin is deprecated; please use the built-in OpenXR plugin or OculusXR plugin from the Marketplace.") EBoundaryType : uint8
 {
 	Boundary_Outer	UMETA(DisplayName = "Outer Boundary"),
 	Boundary_PlayArea	UMETA(DisplayName = "Play Area"),
 };
 
 UENUM(BlueprintType)
-enum class EOculusColorSpace : uint8
+enum class UE_DEPRECATED(5.1, "OculusVR plugin is deprecated; please use the built-in OpenXR plugin or OculusXR plugin from the Marketplace.") EOculusColorSpace : uint8
 {
 	/// The default value from GetHmdColorSpace until SetClientColorDesc is called. Only valid on PC, and will be remapped to Quest on Mobile
 	Unknown = 0,
@@ -121,7 +132,7 @@ enum class EOculusColorSpace : uint8
 };
 
 UENUM(BlueprintType)
-enum class EHandTrackingSupport : uint8
+enum class UE_DEPRECATED(5.1, "OculusVR plugin is deprecated; please use the built-in OpenXR plugin or OculusXR plugin from the Marketplace.") EHandTrackingSupport : uint8
 {
 	ControllersOnly,
 	ControllersAndHands,
@@ -129,7 +140,7 @@ enum class EHandTrackingSupport : uint8
 };
 
 UENUM(BlueprintType)
-enum class EHandTrackingFrequency : uint8
+enum class UE_DEPRECATED(5.1, "OculusVR plugin is deprecated; please use the built-in OpenXR plugin or OculusXR plugin from the Marketplace.") EHandTrackingFrequency : uint8
 {
 	LOW,
 	HIGH,
@@ -137,7 +148,7 @@ enum class EHandTrackingFrequency : uint8
 };
 
 UENUM(BlueprintType)
-enum class EOculusDeviceType : uint8
+enum class UE_DEPRECATED(5.1, "OculusVR plugin is deprecated; please use the built-in OpenXR plugin or OculusXR plugin from the Marketplace.") EOculusDeviceType : uint8
 {
 	//mobile HMDs 
 	OculusMobile_Deprecated0 = 0,
@@ -157,7 +168,7 @@ enum class EOculusDeviceType : uint8
 };
 
 UENUM(BlueprintType)
-enum class EOculusXrApi : uint8
+enum class UE_DEPRECATED(5.1, "OculusVR plugin is deprecated; please use the built-in OpenXR plugin or OculusXR plugin from the Marketplace.") EOculusXrApi : uint8
 {
 	LegacyOVRPlugin = 0 UMETA(DisplayName = "Legacy Oculus SDK (no longer developed by Epic)", ToolTip = "Legacy Oculus SDK. Epic is no longer developing for this SDK, and it should only be used if there are features required for a project that are not yet supported through OpenXR, but it's hard to guarantee potential bugs will be fixed as it's not in active development at Epic. Epic recommends Native OpenXR instead, as that will be the main development focus going forward."),
 	
@@ -171,34 +182,39 @@ enum class EOculusXrApi : uint8
 * EBoundaryType::Boundary_PlayArea) and a device or point in the world.
 * All dimensions, points, and vectors are returned in Unreal world coordinate space.
 */
-USTRUCT(BlueprintType)
+USTRUCT(BlueprintType, meta = (Deprecated = "5.1"))
 struct FGuardianTestResult
 {
 	GENERATED_BODY()
 
 	/** Is there a triggering interaction between the device/point and specified boundary? */
-	UPROPERTY(EditAnywhere, BlueprintReadOnly, Category = "Boundary Test Result")
+	UE_DEPRECATED(5.1, "OculusVR plugin is deprecated; please use the built-in OpenXR plugin or OculusXR plugin from the Marketplace.")
+	UPROPERTY(EditAnywhere, BlueprintReadOnly, Category = "Boundary Test Result", meta = (DeprecatedProperty))
 	bool IsTriggering = false;
 
 	/** Device type triggering boundary (ETrackedDeviceType::None if BoundaryTestResult corresponds to a point rather than a device) */
-	UPROPERTY(EditAnywhere, BlueprintReadOnly, Category = "Boundary Test Result")
+	UE_DEPRECATED(5.1, "OculusVR plugin is deprecated; please use the built-in OpenXR plugin or OculusXR plugin from the Marketplace.")
+	UPROPERTY(EditAnywhere, BlueprintReadOnly, Category = "Boundary Test Result", meta = (DeprecatedProperty))
 	ETrackedDeviceType DeviceType = ETrackedDeviceType::None;
 
 	/** Distance of device/point to surface of boundary specified by BoundaryType */
-	UPROPERTY(EditAnywhere, BlueprintReadOnly, Category = "Boundary Test Result")
+	UE_DEPRECATED(5.1, "OculusVR plugin is deprecated; please use the built-in OpenXR plugin or OculusXR plugin from the Marketplace.")
+	UPROPERTY(EditAnywhere, BlueprintReadOnly, Category = "Boundary Test Result", meta = (DeprecatedProperty))
 	float ClosestDistance = 0.0f;
 
 	/** Closest point on surface corresponding to specified boundary */
-	UPROPERTY(EditAnywhere, BlueprintReadOnly, Category = "Boundary Test Result")
+	UE_DEPRECATED(5.1, "OculusVR plugin is deprecated; please use the built-in OpenXR plugin or OculusXR plugin from the Marketplace.")
+	UPROPERTY(EditAnywhere, BlueprintReadOnly, Category = "Boundary Test Result", meta = (DeprecatedProperty))
 	FVector ClosestPoint = FVector(0.0f);
 
 	/** Normal of closest point */
-	UPROPERTY(EditAnywhere, BlueprintReadOnly, Category = "Boundary Test Result")
+	UE_DEPRECATED(5.1, "OculusVR plugin is deprecated; please use the built-in OpenXR plugin or OculusXR plugin from the Marketplace.")
+	UPROPERTY(EditAnywhere, BlueprintReadOnly, Category = "Boundary Test Result", meta = (DeprecatedProperty))
 	FVector ClosestPointNormal = FVector(0.0f, 0.0f, 1.0f);
 };
 
-UCLASS()
-class OCULUSHMD_API UOculusFunctionLibrary : public UBlueprintFunctionLibrary
+UCLASS(deprecated, meta = (DeprecationMessage = "OculusVR plugin is deprecated; please use the built-in OpenXR plugin or OculusXR plugin from the Marketplace."))
+class OCULUSHMD_API UDEPRECATED_UOculusFunctionLibrary : public UBlueprintFunctionLibrary
 {
 	GENERATED_UCLASS_BODY()
 
@@ -212,7 +228,8 @@ class OCULUSHMD_API UOculusFunctionLibrary : public UBlueprintFunctionLibrary
 	 * @param bUsePositionForPlayerCamera	(in) Should be set to 'true' if the position is going to be used to update position of the camera manually.
 	 * @param PositionScale		(in) The 3D scale that will be applied to position.
 	 */
-	UFUNCTION(BlueprintPure, Category="OculusLibrary")
+	UE_DEPRECATED(5.1, "OculusVR plugin is deprecated; please use the built-in OpenXR plugin or OculusXR plugin from the Marketplace.")
+	UFUNCTION(BlueprintPure, Category="OculusLibrary", meta = (DeprecatedFunction, DeprecationMessage = "The OculusVR plugin is deprecated."))
 	static void GetPose(FRotator& DeviceRotation, FVector& DevicePosition, FVector& NeckPosition, bool bUseOrienationForPlayerCamera = false, bool bUsePositionForPlayerCamera = false, const FVector PositionScale = FVector::ZeroVector);
 
 	/**
@@ -224,18 +241,21 @@ class OCULUSHMD_API UOculusFunctionLibrary : public UBlueprintFunctionLibrary
 	* @param LinearVelocity			(out) Velocity in meters per second.
 	* @param TimeInSeconds			(out) Time when the reported IMU reading took place, in seconds.
 	*/
-	UFUNCTION(BlueprintPure, Category = "OculusLibrary")
+	UE_DEPRECATED(5.1, "OculusVR plugin is deprecated; please use the built-in OpenXR plugin or OculusXR plugin from the Marketplace.")
+	UFUNCTION(BlueprintPure, Category="OculusLibrary", meta = (DeprecatedFunction, DeprecationMessage = "The OculusVR plugin is deprecated."))
 	static void GetRawSensorData(FVector& AngularAcceleration, FVector& LinearAcceleration, FVector& AngularVelocity, FVector& LinearVelocity, float& TimeInSeconds, ETrackedDeviceType DeviceType = ETrackedDeviceType::HMD);
 
 	/**
 	* Returns if the device is currently tracked by the runtime or not.
 	*/
-	UFUNCTION(BlueprintPure, Category = "OculusLibrary")
+	UE_DEPRECATED(5.1, "OculusVR plugin is deprecated; please use the built-in OpenXR plugin or OculusXR plugin from the Marketplace.")
+	UFUNCTION(BlueprintPure, Category="OculusLibrary", meta = (DeprecatedFunction, DeprecationMessage = "The OculusVR plugin is deprecated."))
 	static bool IsDeviceTracked(ETrackedDeviceType DeviceType);
 
 	/**
 	* Returns if the device is currently tracked by the runtime or not.
 	*/
+	UE_DEPRECATED(5.1, "OculusVR plugin is deprecated; please use the built-in OpenXR plugin or OculusXR plugin from the Marketplace.")
 	UFUNCTION(BlueprintCallable, Category = "OculusLibrary")
 	static void SetCPUAndGPULevels(int CPULevel, int GPULevel);
 
@@ -245,7 +265,8 @@ class OCULUSHMD_API UOculusFunctionLibrary : public UBlueprintFunctionLibrary
 	* @param Profile		(out) Structure to hold current user profile.
 	* @return (boolean)	True, if user profile was acquired.
 	*/
-	UFUNCTION(BlueprintPure, Category = "OculusLibrary")
+	UE_DEPRECATED(5.1, "OculusVR plugin is deprecated; please use the built-in OpenXR plugin or OculusXR plugin from the Marketplace.")
+	UFUNCTION(BlueprintPure, Category="OculusLibrary", meta = (DeprecatedFunction, DeprecationMessage = "The OculusVR plugin is deprecated."))
 	static bool GetUserProfile(FHmdUserProfile& Profile);
 
 	/**
@@ -260,7 +281,8 @@ class OCULUSHMD_API UOculusFunctionLibrary : public UBlueprintFunctionLibrary
 	* @param BaseOffsetInMeters (in) the vector to be set as base offset, in meters.
 	* @param Options			(in) specifies either position, orientation or both should be set.
 	*/
-	UFUNCTION(BlueprintCallable, Category = "OculusLibrary")
+	UE_DEPRECATED(5.1, "OculusVR plugin is deprecated; please use the built-in OpenXR plugin or OculusXR plugin from the Marketplace.")
+	UFUNCTION(BlueprintCallable, Category = "OculusLibrary", meta = (DeprecatedFunction, DeprecationMessage = "The OculusVR plugin is deprecated."))
 	static void SetBaseRotationAndBaseOffsetInMeters(FRotator Rotation, FVector BaseOffsetInMeters, EOrientPositionSelector::Type Options);
 
 	/**
@@ -273,7 +295,8 @@ class OCULUSHMD_API UOculusFunctionLibrary : public UBlueprintFunctionLibrary
 	* @param OutRotation			(out) Rotator object with base rotation
 	* @param OutBaseOffsetInMeters	(out) base position offset, vector, in meters.
 	*/
-	UFUNCTION(BlueprintPure, Category = "OculusLibrary")
+	UE_DEPRECATED(5.1, "OculusVR plugin is deprecated; please use the built-in OpenXR plugin or OculusXR plugin from the Marketplace.")
+	UFUNCTION(BlueprintPure, Category="OculusLibrary", meta = (DeprecatedFunction, DeprecationMessage = "The OculusVR plugin is deprecated."))
 	static void GetBaseRotationAndBaseOffsetInMeters(FRotator& OutRotation, FVector& OutBaseOffsetInMeters);
 
 	/**
@@ -281,6 +304,7 @@ class OCULUSHMD_API UOculusFunctionLibrary : public UBlueprintFunctionLibrary
 	 *
 	 * @param PosScale3D	(in) the scale to apply to the HMD position.
 	 */
+	UE_DEPRECATED(5.1, "OculusVR plugin is deprecated; please use the built-in OpenXR plugin or OculusXR plugin from the Marketplace.")
 	UFUNCTION(BlueprintCallable, Category = "OculusLibrary", meta = (DeprecatedFunction, DeprecationMessage = "This feature is no longer supported."))
 	static void SetPositionScale3D(FVector PosScale3D) { }
 
@@ -295,6 +319,7 @@ class OCULUSHMD_API UOculusFunctionLibrary : public UBlueprintFunctionLibrary
 	 * @param PosOffset			(in) the vector to be added to HMD position.
 	 * @param Options			(in) specifies either position, orientation or both should be set.
 	 */
+	UE_DEPRECATED(5.1, "OculusVR plugin is deprecated; please use the built-in OpenXR plugin or OculusXR plugin from the Marketplace.")
 	UFUNCTION(BlueprintCallable, Category = "OculusLibrary", meta = (DeprecatedFunction, DeprecationMessage = "A hack, proper camera positioning should be used"))
 	static void SetBaseRotationAndPositionOffset(FRotator BaseRot, FVector PosOffset, EOrientPositionSelector::Type Options);
 
@@ -304,6 +329,7 @@ class OCULUSHMD_API UOculusFunctionLibrary : public UBlueprintFunctionLibrary
 	 * @param OutRot			(out) Rotator object with base rotation
 	 * @param OutPosOffset		(out) the vector with previously set position offset.
 	 */
+	UE_DEPRECATED(5.1, "OculusVR plugin is deprecated; please use the built-in OpenXR plugin or OculusXR plugin from the Marketplace.")
 	UFUNCTION(BlueprintCallable, Category = "OculusLibrary", meta = (DeprecatedFunction, DeprecationMessage = "A hack, proper camera positioning should be used"))
 	static void GetBaseRotationAndPositionOffset(FRotator& OutRot, FVector& OutPosOffset);
 
@@ -317,49 +343,57 @@ class OCULUSHMD_API UOculusFunctionLibrary : public UBlueprintFunctionLibrary
 	 * @param DeltaRotation		(in) Incremental rotation, that is added each 2nd frame to the quad transform. The quad is rotated around the center of the quad.
 	 * @param bClearBeforeAdd	(in) If true, clears splashes before adding a new one.
 	 */
+	UE_DEPRECATED(5.1, "OculusVR plugin is deprecated; please use the built-in OpenXR plugin or OculusXR plugin from the Marketplace.")
 	UFUNCTION(BlueprintCallable, Category = "OculusLibrary", meta = (DeprecatedFunction, DeprecationMessage = "Use Add Loading Screen Splash from the Head Mounted Display Loading Screen functions instead."))
 	static void AddLoadingSplashScreen(class UTexture2D* Texture, FVector TranslationInMeters, FRotator Rotation, FVector2D SizeInMeters = FVector2D(1.0f, 1.0f), FRotator DeltaRotation = FRotator::ZeroRotator, bool bClearBeforeAdd = false);
 
 	/**
 	 * Removes all the splash screens.
 	 */
+	UE_DEPRECATED(5.1, "OculusVR plugin is deprecated; please use the built-in OpenXR plugin or OculusXR plugin from the Marketplace.")
 	UFUNCTION(BlueprintCallable, Category = "OculusLibrary", meta = (DeprecatedFunction, DeprecationMessage = "Use Clear Loading Screen Splashes from the Head Mounted Display Loading Screen functions instead."))
 	static void ClearLoadingSplashScreens();
 
 	/**
 	* Returns true, if the app has input focus.
 	*/
-	UFUNCTION(BlueprintPure, Category = "OculusLibrary")
+	UE_DEPRECATED(5.1, "OculusVR plugin is deprecated; please use the built-in OpenXR plugin or OculusXR plugin from the Marketplace.")
+	UFUNCTION(BlueprintPure, Category="OculusLibrary", meta = (DeprecatedFunction, DeprecationMessage = "The OculusVR plugin is deprecated."))
 	static bool HasInputFocus();
 
 	/**
 	* Returns true, if the system overlay is present.
 	*/
-	UFUNCTION(BlueprintPure, Category = "OculusLibrary")
+	UE_DEPRECATED(5.1, "OculusVR plugin is deprecated; please use the built-in OpenXR plugin or OculusXR plugin from the Marketplace.")
+	UFUNCTION(BlueprintPure, Category="OculusLibrary", meta = (DeprecatedFunction, DeprecationMessage = "The OculusVR plugin is deprecated."))
 	static bool HasSystemOverlayPresent();
 
 	/**
 	* Returns the GPU utilization availability and value
 	*/
-	UFUNCTION(BlueprintPure, Category = "OculusLibrary")
+	UE_DEPRECATED(5.1, "OculusVR plugin is deprecated; please use the built-in OpenXR plugin or OculusXR plugin from the Marketplace.")
+	UFUNCTION(BlueprintPure, Category="OculusLibrary", meta = (DeprecatedFunction, DeprecationMessage = "The OculusVR plugin is deprecated."))
 	static void GetGPUUtilization(bool& IsGPUAvailable, float& GPUUtilization);
 
 	/**
 	* Returns the GPU frame time on supported mobile platforms (Go for now)
 	*/
-	UFUNCTION(BlueprintPure, Category = "OculusLibrary")
+	UE_DEPRECATED(5.1, "OculusVR plugin is deprecated; please use the built-in OpenXR plugin or OculusXR plugin from the Marketplace.")
+	UFUNCTION(BlueprintPure, Category="OculusLibrary", meta = (DeprecatedFunction, DeprecationMessage = "The OculusVR plugin is deprecated."))
 	static float GetGPUFrameTime();
 
 	/**
 	* Returns the current multiresolution level
 	*/
-	UFUNCTION(BlueprintPure, Category = "OculusLibrary")
+	UE_DEPRECATED(5.1, "OculusVR plugin is deprecated; please use the built-in OpenXR plugin or OculusXR plugin from the Marketplace.")
+	UFUNCTION(BlueprintPure, Category="OculusLibrary", meta = (DeprecatedFunction, DeprecationMessage = "The OculusVR plugin is deprecated."))
 	static EFixedFoveatedRenderingLevel GetFixedFoveatedRenderingLevel();
 
 	/**
 	* Set the requested multiresolution level for the next frame, and whether FFR's level is now dynamic or not.
 	*/
-	UFUNCTION(BlueprintCallable, Category = "OculusLibrary")
+	UE_DEPRECATED(5.1, "OculusVR plugin is deprecated; please use the built-in OpenXR plugin or OculusXR plugin from the Marketplace.")
+	UFUNCTION(BlueprintCallable, Category = "OculusLibrary", meta = (DeprecatedFunction, DeprecationMessage = "The OculusVR plugin is deprecated."))
 	static void SetFixedFoveatedRenderingLevel(EFixedFoveatedRenderingLevel level, bool isDynamic);
 
 	/**
@@ -369,80 +403,93 @@ class OCULUSHMD_API UOculusFunctionLibrary : public UBlueprintFunctionLibrary
 	UFUNCTION(BlueprintPure, Category = "OculusLibrary", meta = (DeprecatedFunction, DeprecationMessage = "UOculusFunctionLibrary::GetDeviceName has been deprecated and no longer functions as before. Please use the enum-based GetDeviceType instead."))
 	static FString GetDeviceName();
 
-	UFUNCTION(BlueprintPure, Category = "OculusLibrary")
+	UE_DEPRECATED(5.1, "OculusVR plugin is deprecated; please use the built-in OpenXR plugin or OculusXR plugin from the Marketplace.")
+	UFUNCTION(BlueprintPure, Category="OculusLibrary", meta = (DeprecatedFunction, DeprecationMessage = "The OculusVR plugin is deprecated."))
 	static EOculusDeviceType GetDeviceType();
 
 	/**
 	* Returns the current available frequencies
 	*/
-	UFUNCTION(BlueprintPure, Category = "OculusLibrary")
+	UE_DEPRECATED(5.1, "OculusVR plugin is deprecated; please use the built-in OpenXR plugin or OculusXR plugin from the Marketplace.")
+	UFUNCTION(BlueprintPure, Category="OculusLibrary", meta = (DeprecatedFunction, DeprecationMessage = "The OculusVR plugin is deprecated."))
 	static TArray<float> GetAvailableDisplayFrequencies();
 
 	/**
 	* Returns the current display frequency
 	*/
-	UFUNCTION(BlueprintPure, Category = "OculusLibrary")
+	UE_DEPRECATED(5.1, "OculusVR plugin is deprecated; please use the built-in OpenXR plugin or OculusXR plugin from the Marketplace.")
+	UFUNCTION(BlueprintPure, Category="OculusLibrary", meta = (DeprecatedFunction, DeprecationMessage = "The OculusVR plugin is deprecated."))
 	static float GetCurrentDisplayFrequency();
 
 	/**
 	* Sets the requested display frequency
 	*/
-	UFUNCTION(BlueprintCallable, Category = "OculusLibrary")
+	UE_DEPRECATED(5.1, "OculusVR plugin is deprecated; please use the built-in OpenXR plugin or OculusXR plugin from the Marketplace.")
+	UFUNCTION(BlueprintCallable, Category = "OculusLibrary", meta = (DeprecatedFunction, DeprecationMessage = "The OculusVR plugin is deprecated."))
 	static void SetDisplayFrequency(float RequestedFrequency);
 
 	/**
 	* Enables/disables positional tracking on devices that support it.
 	*/
-	UFUNCTION(BlueprintCallable, Category = "OculusLibrary")
+	UE_DEPRECATED(5.1, "OculusVR plugin is deprecated; please use the built-in OpenXR plugin or OculusXR plugin from the Marketplace.")
+	UFUNCTION(BlueprintCallable, Category = "OculusLibrary", meta = (DeprecatedFunction, DeprecationMessage = "The OculusVR plugin is deprecated."))
 	static void EnablePositionTracking(bool bPositionTracking);
 
 	/**
 	* Enables/disables orientation tracking on devices that support it.
 	*/
-	UFUNCTION(BlueprintCallable, Category = "OculusLibrary")
+	UE_DEPRECATED(5.1, "OculusVR plugin is deprecated; please use the built-in OpenXR plugin or OculusXR plugin from the Marketplace.")
+	UFUNCTION(BlueprintCallable, Category = "OculusLibrary", meta = (DeprecatedFunction, DeprecationMessage = "The OculusVR plugin is deprecated."))
 	static void EnableOrientationTracking(bool bOrientationTracking);
 
 	/**
 	* Set the Color Scale/Offset
 	*/
-	UFUNCTION(BlueprintCallable, Category = "OculusLibrary")
+	UE_DEPRECATED(5.1, "OculusVR plugin is deprecated; please use the built-in OpenXR plugin or OculusXR plugin from the Marketplace.")
+	UFUNCTION(BlueprintCallable, Category = "OculusLibrary", meta = (DeprecatedFunction, DeprecationMessage = "The OculusVR plugin is deprecated."))
 	static void SetColorScaleAndOffset(FLinearColor ColorScale, FLinearColor ColorOffset, bool bApplyToAllLayers = false);
 
 	/**
 	* Returns true if system headset is in 3dof mode 
 	*/
-	UFUNCTION(BlueprintPure, Category = "OculusLibrary")
+	UE_DEPRECATED(5.1, "OculusVR plugin is deprecated; please use the built-in OpenXR plugin or OculusXR plugin from the Marketplace.")
+	UFUNCTION(BlueprintPure, Category="OculusLibrary", meta = (DeprecatedFunction, DeprecationMessage = "The OculusVR plugin is deprecated."))
 	static bool GetSystemHmd3DofModeEnabled();
 
 	/**
 	* Returns the color space of the target HMD
 	*/
-	UFUNCTION(BlueprintPure, Category = "OculusLibrary")
+	UE_DEPRECATED(5.1, "OculusVR plugin is deprecated; please use the built-in OpenXR plugin or OculusXR plugin from the Marketplace.")
+	UFUNCTION(BlueprintPure, Category="OculusLibrary", meta = (DeprecatedFunction, DeprecationMessage = "The OculusVR plugin is deprecated."))
 	static EOculusColorSpace GetHmdColorDesc();
 
 	/**
 	* Sets the target HMD to do color space correction to a specific color space
 	*/
-	UFUNCTION(BlueprintCallable, Category = "OculusLibrary")
+	UE_DEPRECATED(5.1, "OculusVR plugin is deprecated; please use the built-in OpenXR plugin or OculusXR plugin from the Marketplace.")
+	UFUNCTION(BlueprintCallable, Category = "OculusLibrary", meta = (DeprecatedFunction, DeprecationMessage = "The OculusVR plugin is deprecated."))
 	static void SetClientColorDesc(EOculusColorSpace ColorSpace);
 
 	/**
 	 * Returns IStereoLayers interface to work with overlays.
 	 */
+	UE_DEPRECATED(5.1, "OculusVR plugin is deprecated; please use the built-in OpenXR plugin or OculusXR plugin from the Marketplace.")
 	static class IStereoLayers* GetStereoLayers();
 
 	/* GUARDIAN API */
 	/**
 	* Returns true if the Guardian Outer Boundary is being displayed
 	*/
-	UFUNCTION(BlueprintPure, Category = "OculusLibrary|Guardian")
+	UE_DEPRECATED(5.1, "OculusVR plugin is deprecated; please use the built-in OpenXR plugin or OculusXR plugin from the Marketplace.")
+	UFUNCTION(BlueprintPure, Category = "OculusLibrary|Guardian", meta = (DeprecatedFunction, DeprecationMessage = "The OculusVR plugin is deprecated."))
 	static bool IsGuardianDisplayed();
 
 	/* GUARDIAN API */
 	/**
 	* Returns true if the Guardian has been set up by the user, false if the user is in "seated" mode and has not set up a play space.
 	*/
-	UFUNCTION(BlueprintPure, Category = "OculusLibrary|Guardian")
+	UE_DEPRECATED(5.1, "OculusVR plugin is deprecated; please use the built-in OpenXR plugin or OculusXR plugin from the Marketplace.")
+	UFUNCTION(BlueprintPure, Category = "OculusLibrary|Guardian", meta = (DeprecatedFunction, DeprecationMessage = "The OculusVR plugin is deprecated."))
 	static bool IsGuardianConfigured();
 
 	/**
@@ -450,20 +497,23 @@ class OCULUSHMD_API UOculusFunctionLibrary : public UBlueprintFunctionLibrary
 	* @param BoundaryType			(in) An enum representing the boundary type requested, either Outer Boundary (exact guardian bounds) or PlayArea (rectangle inside the Outer Boundary)
 	* @param UsePawnSpace			(in) Boolean indicating to return the points in world space or pawn space
 	*/
-	UFUNCTION(BlueprintPure, Category = "OculusLibrary|Guardian")
+	UE_DEPRECATED(5.1, "OculusVR plugin is deprecated; please use the built-in OpenXR plugin or OculusXR plugin from the Marketplace.")
+	UFUNCTION(BlueprintPure, Category = "OculusLibrary|Guardian", meta = (DeprecatedFunction, DeprecationMessage = "The OculusVR plugin is deprecated."))
 	static TArray<FVector> GetGuardianPoints(EBoundaryType BoundaryType, bool UsePawnSpace = false);
 
 	/**
 	* Returns the dimensions in UE world space of the requested Boundary Type
 	* @param BoundaryType			(in) An enum representing the boundary type requested, either Outer Boundary (exact guardian bounds) or PlayArea (rectangle inside the Outer Boundary)
 	*/
-	UFUNCTION(BlueprintPure, Category = "OculusLibrary|Guardian")
+	UE_DEPRECATED(5.1, "OculusVR plugin is deprecated; please use the built-in OpenXR plugin or OculusXR plugin from the Marketplace.")
+	UFUNCTION(BlueprintPure, Category = "OculusLibrary|Guardian", meta = (DeprecatedFunction, DeprecationMessage = "The OculusVR plugin is deprecated."))
 	static FVector GetGuardianDimensions(EBoundaryType BoundaryType);
 
 	/**
 	* Returns the transform of the play area rectangle, defining its position, rotation and scale to apply to a unit cube to match it with the play area.
 	*/
-	UFUNCTION(BlueprintPure, Category = "OculusLibrary|Guardian")
+	UE_DEPRECATED(5.1, "OculusVR plugin is deprecated; please use the built-in OpenXR plugin or OculusXR plugin from the Marketplace.")
+	UFUNCTION(BlueprintPure, Category = "OculusLibrary|Guardian", meta = (DeprecatedFunction, DeprecationMessage = "The OculusVR plugin is deprecated."))
 	static FTransform GetPlayAreaTransform();
 
 	/**
@@ -471,7 +521,8 @@ class OCULUSHMD_API UOculusFunctionLibrary : public UBlueprintFunctionLibrary
 	* @param Point					(in) Point in UE space to test against guardian boundaries
 	* @param BoundaryType			(in) An enum representing the boundary type requested, either Outer Boundary (exact guardian bounds) or PlayArea (rectangle inside the Outer Boundary)
 	*/
-	UFUNCTION(BlueprintCallable, Category = "OculusLibrary|Guardian")
+	UE_DEPRECATED(5.1, "OculusVR plugin is deprecated; please use the built-in OpenXR plugin or OculusXR plugin from the Marketplace.")
+	UFUNCTION(BlueprintCallable, Category = "OculusLibrary|Guardian", meta = (DeprecatedFunction, DeprecationMessage = "The OculusVR plugin is deprecated."))
 	static FGuardianTestResult GetPointGuardianIntersection(const FVector Point, EBoundaryType BoundaryType);
 
 	/**
@@ -479,14 +530,16 @@ class OCULUSHMD_API UOculusFunctionLibrary : public UBlueprintFunctionLibrary
 	* @param DeviceType             (in) Tracked Device type to test against guardian boundaries
 	* @param BoundaryType			(in) An enum representing the boundary type requested, either Outer Boundary (exact guardian bounds) or PlayArea (rectangle inside the Outer Boundary)
 	*/
-	UFUNCTION(BlueprintCallable, Category = "OculusLibrary|Guardian")
+	UE_DEPRECATED(5.1, "OculusVR plugin is deprecated; please use the built-in OpenXR plugin or OculusXR plugin from the Marketplace.")
+	UFUNCTION(BlueprintCallable, Category = "OculusLibrary|Guardian", meta = (DeprecatedFunction, DeprecationMessage = "The OculusVR plugin is deprecated."))
 	static FGuardianTestResult GetNodeGuardianIntersection(ETrackedDeviceType DeviceType, EBoundaryType BoundaryType);
 
 	/**
 	* Forces the runtime to render guardian at all times or not
 	* @param GuardianVisible			(in) True will display guardian, False will hide it
 	*/
-	UFUNCTION(BlueprintCallable, Category = "OculusLibrary|Guardian")
+	UE_DEPRECATED(5.1, "OculusVR plugin is deprecated; please use the built-in OpenXR plugin or OculusXR plugin from the Marketplace.")
+	UFUNCTION(BlueprintCallable, Category = "OculusLibrary|Guardian", meta = (DeprecatedFunction, DeprecationMessage = "The OculusVR plugin is deprecated."))
 	static void SetGuardianVisibility(bool GuardianVisible);
 
 	/** When player triggers the Guardian boundary */
@@ -509,3 +562,5 @@ class OCULUSHMD_API UOculusFunctionLibrary : public UBlueprintFunctionLibrary
 protected:
 	static class OculusHMD::FOculusHMD* GetOculusHMD();
 };
+
+PRAGMA_ENABLE_DEPRECATION_WARNINGS
diff --git a/Engine/Plugins/Runtime/Oculus/OculusVR/Source/OculusHMD/Public/OculusHMDRuntimeSettings.h b/Engine/Plugins/Runtime/Oculus/OculusVR/Source/OculusHMD/Public/OculusHMDRuntimeSettings.h
--- a/Engine/Plugins/Runtime/Oculus/OculusVR/Source/OculusHMD/Public/OculusHMDRuntimeSettings.h
+++ b/Engine/Plugins/Runtime/Oculus/OculusVR/Source/OculusHMD/Public/OculusHMDRuntimeSettings.h
@@ -8,86 +8,106 @@
 #include "OculusFunctionLibrary.h"
 #include "OculusHMDRuntimeSettings.generated.h"
 
+PRAGMA_DISABLE_DEPRECATION_WARNINGS
+
 /**
 * Implements the settings for the OculusVR plugin.
 */
-UCLASS(config = Engine, defaultconfig)
-class OCULUSHMD_API UOculusHMDRuntimeSettings : public UObject
+UCLASS(config = Engine, defaultconfig, deprecated, meta = (DeprecationMessage = "OculusVR plugin is deprecated; please use the built-in OpenXR plugin or OculusXR plugin from the Marketplace."))
+class OCULUSHMD_API UDEPRECATED_UOculusHMDRuntimeSettings : public UObject
 {
 	GENERATED_UCLASS_BODY()
 
 public:
 	
 	/** Whether the Splash screen is enabled. */
-	UPROPERTY(config, EditAnywhere, Category = "Engine SplashScreen")
+	UE_DEPRECATED(5.1, "OculusVR plugin is deprecated; please use the built-in OpenXR plugin or OculusXR plugin from the Marketplace.")
+	UPROPERTY(config, EditAnywhere, Category = "Engine SplashScreen", meta = (DeprecatedProperty))
 	bool bAutoEnabled;
 
 	/** An array of splash screen descriptors listing textures to show and their positions. */
-	UPROPERTY(config, EditAnywhere, Category = "Engine SplashScreen")
+	UE_DEPRECATED(5.1, "OculusVR plugin is deprecated; please use the built-in OpenXR plugin or OculusXR plugin from the Marketplace.")
+	UPROPERTY(config, EditAnywhere, Category = "Engine SplashScreen", meta = (DeprecatedProperty))
 	TArray<FOculusSplashDesc> SplashDescs;
 
 	/** This selects the XR API that the engine will use. If unsure, OVRPlugin OpenXR is the recommended API. */
-	UPROPERTY(config, EditAnywhere, Category = General, meta = (DisplayName = "XR API", ConfigRestartRequired = true))
+	UE_DEPRECATED(5.1, "OculusVR plugin is deprecated; please use the built-in OpenXR plugin or OculusXR plugin from the Marketplace.")
+	UPROPERTY(config, EditAnywhere, Category = General, meta = (DisplayName = "XR API", ConfigRestartRequired = true, DeprecatedProperty))
 	EOculusXrApi XrApi;
 
 	/** The target color space */
-	UPROPERTY(config, EditAnywhere, Category = General)
+	UE_DEPRECATED(5.1, "OculusVR plugin is deprecated; please use the built-in OpenXR plugin or OculusXR plugin from the Marketplace.")
+	UPROPERTY(config, EditAnywhere, Category = General, meta = (DeprecatedProperty))
 	EOculusColorSpace	ColorSpace;
 
 	/** Whether Dash is supported by the app, which will keep the app in foreground when the User presses the oculus button (needs the app to handle input focus loss!) */
-	UPROPERTY(config, EditAnywhere, Category = PC)
+	UE_DEPRECATED(5.1, "OculusVR plugin is deprecated; please use the built-in OpenXR plugin or OculusXR plugin from the Marketplace.")
+	UPROPERTY(config, EditAnywhere, Category = PC, meta = (DeprecatedProperty))
 	bool bSupportsDash;
 
 	/** Whether the app's depth buffer is shared with the Rift Compositor, for layer (including Dash) compositing, PTW, and potentially more. */
-	UPROPERTY(config, EditAnywhere, Category = PC)
+	UE_DEPRECATED(5.1, "OculusVR plugin is deprecated; please use the built-in OpenXR plugin or OculusXR plugin from the Marketplace.")
+	UPROPERTY(config, EditAnywhere, Category = PC, meta = (DeprecatedProperty))
 	bool bCompositesDepth;
 
 	/** Computes mipmaps for the eye buffers every frame, for a higher quality distortion */
-	UPROPERTY(config, EditAnywhere, Category = PC)
+	UE_DEPRECATED(5.1, "OculusVR plugin is deprecated; please use the built-in OpenXR plugin or OculusXR plugin from the Marketplace.")
+	UPROPERTY(config, EditAnywhere, Category = PC, meta = (DeprecatedProperty))
 	bool bHQDistortion;
 
 	/** Minimum allowed pixel density. */
-	UPROPERTY(config, EditAnywhere, Category = PC)
+	UE_DEPRECATED(5.1, "OculusVR plugin is deprecated; please use the built-in OpenXR plugin or OculusXR plugin from the Marketplace.")
+	UPROPERTY(config, EditAnywhere, Category = PC, meta = (DeprecatedProperty))
 	float PixelDensityMin;
 
 	/** Maximum allowed pixel density. */
-	UPROPERTY(config, EditAnywhere, Category = PC)
+	UE_DEPRECATED(5.1, "OculusVR plugin is deprecated; please use the built-in OpenXR plugin or OculusXR plugin from the Marketplace.")
+	UPROPERTY(config, EditAnywhere, Category = PC, meta = (DeprecatedProperty))
 	float PixelDensityMax;
 
 	/** A png for Mobile-OS-driven launch splash screen. It will show up instantly at app launch and disappear upon first engine-driven frame (regardless of said frame being UE4 splashes or 3D scenes) */
-	UPROPERTY(config, EditAnywhere, Category = Mobile, meta = (DisplayName = "OS Splash Screen", FilePathFilter = "png", RelativeToGameDir))
+	UE_DEPRECATED(5.1, "OculusVR plugin is deprecated; please use the built-in OpenXR plugin or OculusXR plugin from the Marketplace.")
+	UPROPERTY(config, EditAnywhere, Category = Mobile, meta = (DisplayName = "OS Splash Screen", FilePathFilter = "png", RelativeToGameDir, DeprecatedProperty))
 	FFilePath OSSplashScreen;
 
 	/** Default CPU level controlling CPU frequency on the mobile device */
-	UPROPERTY(config, EditAnywhere, Category = Mobile)
+	UE_DEPRECATED(5.1, "OculusVR plugin is deprecated; please use the built-in OpenXR plugin or OculusXR plugin from the Marketplace.")
+	UPROPERTY(config, EditAnywhere, Category = Mobile, meta = (DeprecatedProperty))
 	int CPULevel;
 
 	/** Default GPU level controlling GPU frequency on the mobile device */
-	UPROPERTY(config, EditAnywhere, Category = Mobile)
+	UE_DEPRECATED(5.1, "OculusVR plugin is deprecated; please use the built-in OpenXR plugin or OculusXR plugin from the Marketplace.")
+	UPROPERTY(config, EditAnywhere, Category = Mobile, meta = (DeprecatedProperty))
 	int GPULevel;
 
 	/** If enabled the app will be focus aware. This will keep the app in foreground when the User presses the oculus button (needs the app to handle input focus loss!) */
-	UPROPERTY(config, EditAnywhere, Category = Mobile)
+	UE_DEPRECATED(5.1, "OculusVR plugin is deprecated; please use the built-in OpenXR plugin or OculusXR plugin from the Marketplace.")
+	UPROPERTY(config, EditAnywhere, Category = Mobile, meta = (DeprecatedProperty))
 	bool bFocusAware;
 
 	/** [Experimental]Enable Late latching for reducing HMD and controller latency, improve tracking prediction quality, multiview and vulkan must be enabled for this feature. */
-	UPROPERTY(config, EditAnywhere, Category = Mobile)
+	UE_DEPRECATED(5.1, "OculusVR plugin is deprecated; please use the built-in OpenXR plugin or OculusXR plugin from the Marketplace.")
+	UPROPERTY(config, EditAnywhere, Category = Mobile, meta = (DeprecatedProperty))
 	bool bLateLatching;
 
 	/** If enabled the app will use the Oculus system keyboard for input fields. This requires that the app be focus aware. */
-	UPROPERTY(config, EditAnywhere, Category = Mobile)
+	UE_DEPRECATED(5.1, "OculusVR plugin is deprecated; please use the built-in OpenXR plugin or OculusXR plugin from the Marketplace.")
+	UPROPERTY(config, EditAnywhere, Category = Mobile, meta = (DeprecatedProperty))
 	bool bRequiresSystemKeyboard;
 
 	/** Whether controllers and/or hands can be used with the app */
-	UPROPERTY(config, EditAnywhere, Category = Mobile)
+	UE_DEPRECATED(5.1, "OculusVR plugin is deprecated; please use the built-in OpenXR plugin or OculusXR plugin from the Marketplace.")
+	UPROPERTY(config, EditAnywhere, Category = Mobile, meta = (DeprecatedProperty))
 	EHandTrackingSupport HandTrackingSupport;
 
 	/** Note that a higher tracking frequency will reserve some performance headroom from the application's budget. */
-	UPROPERTY(config, EditAnywhere, Category = Mobile)
+	UE_DEPRECATED(5.1, "OculusVR plugin is deprecated; please use the built-in OpenXR plugin or OculusXR plugin from the Marketplace.")
+	UPROPERTY(config, EditAnywhere, Category = Mobile, meta = (DeprecatedProperty))
 	EHandTrackingFrequency HandTrackingFrequency;
 
 	/** Enable phase sync on mobile, reducing HMD and controller latency, improve tracking prediction quality */
-	UPROPERTY(config, EditAnywhere, Category = Mobile)
+	UE_DEPRECATED(5.1, "OculusVR plugin is deprecated; please use the built-in OpenXR plugin or OculusXR plugin from the Marketplace.")
+	UPROPERTY(config, EditAnywhere, Category = Mobile, meta = (DeprecatedProperty))
 	bool bPhaseSync;
 
 #if WITH_EDITOR
@@ -98,3 +118,4 @@ private:
 	void LoadFromIni();
 
 };
+PRAGMA_ENABLE_DEPRECATION_WARNINGS
diff --git a/Engine/Plugins/Runtime/Oculus/OculusVR/Source/OculusHMD/Public/OculusHMDTypes.h b/Engine/Plugins/Runtime/Oculus/OculusVR/Source/OculusHMD/Public/OculusHMDTypes.h
--- a/Engine/Plugins/Runtime/Oculus/OculusVR/Source/OculusHMD/Public/OculusHMDTypes.h
+++ b/Engine/Plugins/Runtime/Oculus/OculusVR/Source/OculusHMD/Public/OculusHMDTypes.h
@@ -10,38 +10,54 @@
 #include "UObject/SoftObjectPath.h"
 #include "OculusHMDTypes.generated.h"
 
-USTRUCT()
+PRAGMA_DISABLE_DEPRECATION_WARNINGS
+
+USTRUCT(meta = (Deprecated = "5.1"))
 struct FOculusSplashDesc
 {
 	GENERATED_USTRUCT_BODY()
 
+	UE_DEPRECATED(5.1, "OculusVR plugin is deprecated; please use the built-in OpenXR plugin or OculusXR plugin from the Marketplace.")
 	UPROPERTY(config, EditAnywhere, Category = Settings, meta = (
 		AllowedClasses = "/Script/Engine.Texture",
-		ToolTip = "Texture to display"))
+		ToolTip = "Texture to display",
+		DeprecatedProperty))
 	FSoftObjectPath		TexturePath;
 
+	UE_DEPRECATED(5.1, "OculusVR plugin is deprecated; please use the built-in OpenXR plugin or OculusXR plugin from the Marketplace.")
 	UPROPERTY(config, EditAnywhere, Category = Settings, meta = (
-		ToolTip = "transform of center of quad (meters)."))
+		ToolTip = "transform of center of quad (meters).",
+		DeprecatedProperty))
 	FTransform			TransformInMeters;
 
+	UE_DEPRECATED(5.1, "OculusVR plugin is deprecated; please use the built-in OpenXR plugin or OculusXR plugin from the Marketplace.")
 	UPROPERTY(config, EditAnywhere, Category = Settings, meta = (
-		ToolTip = "Dimensions in meters."))
+		ToolTip = "Dimensions in meters.",
+		DeprecatedProperty))
 	FVector2D			QuadSizeInMeters;
 
+	UE_DEPRECATED(5.1, "OculusVR plugin is deprecated; please use the built-in OpenXR plugin or OculusXR plugin from the Marketplace.")
 	UPROPERTY(config, EditAnywhere, Category = Settings, meta = (
-		ToolTip = "A delta rotation that will be added each rendering frame (half rate of full vsync)."))
+		ToolTip = "A delta rotation that will be added each rendering frame (half rate of full vsync).",
+		DeprecatedProperty))
 	FQuat				DeltaRotation;
 
+	UE_DEPRECATED(5.1, "OculusVR plugin is deprecated; please use the built-in OpenXR plugin or OculusXR plugin from the Marketplace.")
 	UPROPERTY(config, EditAnywhere, Category = Settings, meta = (
-		ToolTip = "Texture offset amount from the top left corner."))
+		ToolTip = "Texture offset amount from the top left corner.",
+		DeprecatedProperty))
 	FVector2D			TextureOffset;
 
+	UE_DEPRECATED(5.1, "OculusVR plugin is deprecated; please use the built-in OpenXR plugin or OculusXR plugin from the Marketplace.")
 	UPROPERTY(config, EditAnywhere, Category = Settings, meta = (
-		ToolTip = "Texture scale."))
+		ToolTip = "Texture scale.",
+		DeprecatedProperty))
 	FVector2D			TextureScale;
 
+	UE_DEPRECATED(5.1, "OculusVR plugin is deprecated; please use the built-in OpenXR plugin or OculusXR plugin from the Marketplace.")
 	UPROPERTY(config, EditAnywhere, Category = Settings, meta = (
-		ToolTip = "Whether the splash layer uses it's alpha channel."))
+		ToolTip = "Whether the splash layer uses it's alpha channel.",
+		DeprecatedProperty))
 	bool				bNoAlphaChannel;
 
 	// Runtime data
@@ -72,3 +88,4 @@ struct FOculusSplashDesc
 			LoadingTexture == d.LoadingTexture && LoadedTexture == d.LoadedTexture && bIsDynamic == d.bIsDynamic;
 	}
 };
+PRAGMA_ENABLE_DEPRECATION_WARNINGS
diff --git a/Engine/Plugins/Runtime/Oculus/OculusVR/Source/OculusInput/Private/OculusInputModule.cpp b/Engine/Plugins/Runtime/Oculus/OculusVR/Source/OculusInput/Private/OculusInputModule.cpp
--- a/Engine/Plugins/Runtime/Oculus/OculusVR/Source/OculusInput/Private/OculusInputModule.cpp
+++ b/Engine/Plugins/Runtime/Oculus/OculusVR/Source/OculusInput/Private/OculusInputModule.cpp
@@ -13,6 +13,8 @@
 // FOculusInputModule
 //-------------------------------------------------------------------------------------------------
 
+PRAGMA_DISABLE_DEPRECATION_WARNINGS
+
 void FOculusInputModule::StartupModule()
 {
 	IInputDeviceModule::StartupModule();
@@ -66,6 +68,7 @@ TSharedPtr<IInputDevice> FOculusInputModule::GetInputDevice() const
 	return NULL;
 }
 
+PRAGMA_ENABLE_DEPRECATION_WARNINGS
 
 #endif	// OCULUS_INPUT_SUPPORTED_PLATFORMS
 
diff --git a/Engine/Plugins/Runtime/Oculus/OculusVR/Source/OculusMR/Private/OculusMR_State.h b/Engine/Plugins/Runtime/Oculus/OculusVR/Source/OculusMR/Private/OculusMR_State.h
--- a/Engine/Plugins/Runtime/Oculus/OculusVR/Source/OculusMR/Private/OculusMR_State.h
+++ b/Engine/Plugins/Runtime/Oculus/OculusVR/Source/OculusMR/Private/OculusMR_State.h
@@ -7,6 +7,8 @@
 
 #include "OculusMR_State.generated.h"
 
+PRAGMA_DISABLE_DEPRECATION_WARNINGS
+
 USTRUCT()
 struct FTrackedCamera
 {
@@ -115,4 +117,6 @@ public:
 	/** Flag indicating a change in the tracked camera index for the camera actor to consume */
 	UPROPERTY()
 	bool BindToTrackedCameraIndexRequested;
-};
\ No newline at end of file
+};
+
+PRAGMA_ENABLE_DEPRECATION_WARNINGS
