OculusVR deprecation: OculusVR plugin deprecation

diff --git a/Engine/Plugins/Runtime/Oculus/OculusVR/OculusVR.uplugin b/Engine/Plugins/Runtime/Oculus/OculusVR/OculusVR.uplugin
--- a/Engine/Plugins/Runtime/Oculus/OculusVR/OculusVR.uplugin
+++ b/Engine/Plugins/Runtime/Oculus/OculusVR/OculusVR.uplugin
@@ -3,7 +3,7 @@
 	"Version": 1,
 	"VersionName": "1.59.0",
 	"FriendlyName": "Oculus VR",
-	"Description": "Support for Oculus headsets and controllers",
+	"Description": "Support for Oculus headsets and controllers. This plugin is deprecated and will be removed in a future engine release. Please use the built-in OpenXR plugin or marketplace OculusXR plugin.",
 	"Category": "Virtual Reality",
 	"CreatedBy": "Oculus",
 	"CreatedByURL": "https://www.oculus.com/",
