PR #9561: Remove FFR redirect from BaseEngine, should be in plugin (Contributed by [REDACTED_PERSON])

diff --git a/Engine/Config/BaseEngine.ini b/Engine/Config/BaseEngine.ini
--- a/Engine/Config/BaseEngine.ini
+++ b/Engine/Config/BaseEngine.ini
@@ -680,9 +680,6 @@ DefaultWorldPartitionHLODLayer=/Engine/Maps/Templates/HLODs/HLODLayer_Instanced.
 
 +PropertyRedirects=(OldName="GameplayStatics.CreatePlayer.bSpawnPawn",NewName="bSpawnPlayerController")
 
-+FunctionRedirects=(OldName="GetTiledMultiresLevel",NewName="GetFixedFoveatedRenderingLevel")
-+FunctionRedirects=(OldName="SetTiledMultiresLevel",NewName="SetFixedFoveatedRenderingLevel")
-
 ; 4.23
 +FunctionRedirects=(OldName="SkeletalMeshComponent.SetAnimInstanceClass",NewName="SkeletalMeshComponent.SetAnimClass")
 +FunctionRedirects=(OldName="SkeletalMeshComponent.K2_SetAnimInstanceClass",NewName="SkeletalMeshComponent.SetAnimClass")
diff --git a/Engine/Plugins/Runtime/Oculus/OculusVR/Config/BaseOculusVR.ini b/Engine/Plugins/Runtime/Oculus/OculusVR/Config/BaseOculusVR.ini
--- a/Engine/Plugins/Runtime/Oculus/OculusVR/Config/BaseOculusVR.ini
+++ b/Engine/Plugins/Runtime/Oculus/OculusVR/Config/BaseOculusVR.ini
@@ -8,6 +8,8 @@
 ; OculusRift module
 +PackageRedirects=(OldName="/Script/OculusRift", NewName="/Script/OculusHMD")
 +EnumRedirects=(OldName="ETiledMultiResLevel",NewName="/Script/OculusHMD.EFixedFoveatedRenderingLevel",ValueChanges=(("ETiledMultiResLevel_Off","FFR_Off"),("ETiledMultiResLevel_LMSLow","FFR_Low"),("ETiledMultiResLevel_LMSMedium","FFR_Medium"),("ETiledMultiResLevel_LMSHigh","FFR_High"),("ETiledMultiResLevel_LMSHighTop","FFR_HighTop")))
++FunctionRedirects=(OldName="GetTiledMultiresLevel",NewName="/Script/OculusHMD.GetFixedFoveatedRenderingLevel")
++FunctionRedirects=(OldName="SetTiledMultiresLevel",NewName="/Script/OculusHMD.SetFixedFoveatedRenderingLevel")
 
 ; 5.0
 
