Clearer deprecations of Oculus VR + SteamVR

diff --git a/Engine/Plugins/Runtime/Oculus/OculusVR/OculusVR.uplugin b/Engine/Plugins/Runtime/Oculus/OculusVR/OculusVR.uplugin
--- a/Engine/Plugins/Runtime/Oculus/OculusVR/OculusVR.uplugin
+++ b/Engine/Plugins/Runtime/Oculus/OculusVR/OculusVR.uplugin
@@ -2,7 +2,7 @@
 	"FileVersion": 3,
 	"Version": 1,
 	"VersionName": "1.59.0",
-	"FriendlyName": "Oculus VR",
+	"FriendlyName": "Oculus VR (deprecated)",
 	"Description": "Support for Oculus headsets and controllers. This plugin is deprecated and will be removed in a future engine release. Please use the built-in OpenXR plugin or marketplace OculusXR plugin.",
 	"Category": "Virtual Reality",
 	"CreatedBy": "Oculus",
diff --git a/Engine/Plugins/Runtime/Steam/SteamVR/SteamVR.uplugin b/Engine/Plugins/Runtime/Steam/SteamVR/SteamVR.uplugin
--- a/Engine/Plugins/Runtime/Steam/SteamVR/SteamVR.uplugin
+++ b/Engine/Plugins/Runtime/Steam/SteamVR/SteamVR.uplugin
@@ -2,7 +2,7 @@
 	"FileVersion": 3,
 	"Version": 1,
 	"VersionName": "1.0",
-	"FriendlyName": "SteamVR",
+	"FriendlyName": "SteamVR (deprecated)",
 	"Description": "Support for the SteamVR API. This plugin is deprecated and will be removed in a future engine release. Please use the OpenXR plugin.",
 	"Category": "Virtual Reality",
 	"CreatedBy": "Epic Games, Inc.",
