Fix non-unity error in OculusHandTracking.cpp introduced by 21089728 referencing GetSkinnedAsset instead of skeletal meshcomponent

#rnx

diff --git a/Engine/Plugins/Runtime/Oculus/OculusVR/Source/OculusInput/Private/OculusHandTracking.cpp b/Engine/Plugins/Runtime/Oculus/OculusVR/Source/OculusInput/Private/OculusHandTracking.cpp
--- a/Engine/Plugins/Runtime/Oculus/OculusVR/Source/OculusInput/Private/OculusHandTracking.cpp
+++ b/Engine/Plugins/Runtime/Oculus/OculusVR/Source/OculusInput/Private/OculusHandTracking.cpp
@@ -14,6 +14,7 @@
 #include "Materials/MaterialInstanceDynamic.h"
 #include "Model.h"
 #include "GenericPlatform/GenericPlatformInputDeviceMapper.h"
+#include "Engine/SkinnedAssetCommon.h"
 
 #define OCULUS_TO_UE4_SCALE 100.0f
 
