Fix non-unity compile error in OculusHMD_Layer.cpp (UE-168763).

diff --git a/Engine/Plugins/Runtime/Oculus/OculusVR/Source/OculusHMD/Private/OculusHMD_Layer.cpp b/Engine/Plugins/Runtime/Oculus/OculusVR/Source/OculusHMD/Private/OculusHMD_Layer.cpp
--- a/Engine/Plugins/Runtime/Oculus/OculusVR/Source/OculusHMD/Private/OculusHMD_Layer.cpp
+++ b/Engine/Plugins/Runtime/Oculus/OculusVR/Source/OculusHMD/Private/OculusHMD_Layer.cpp
@@ -10,8 +10,8 @@
 #include "Engine/GameEngine.h"
 #include "Engine/Texture2D.h"
 #include "HeadMountedDisplayTypes.h" // for LogHMD
+#include "Materials/Material.h"
 #include "Materials/MaterialInstanceDynamic.h"
-#include "Materials/MaterialInterface.h"
 #include "OculusHMD.h"
 #include "OculusHMDModule.h"
 #include "OculusHMDPrivate.h"
