Added forward declaration in OculusMR_CastingCameraActor to fix non-unity/pch compile error on linux

diff --git a/Engine/Plugins/Runtime/Oculus/OculusVR/Source/OculusMR/Private/OculusMR_CastingCameraActor.h b/Engine/Plugins/Runtime/Oculus/OculusVR/Source/OculusMR/Private/OculusMR_CastingCameraActor.h
--- a/Engine/Plugins/Runtime/Oculus/OculusVR/Source/OculusMR/Private/OculusMR_CastingCameraActor.h
+++ b/Engine/Plugins/Runtime/Oculus/OculusVR/Source/OculusMR/Private/OculusMR_CastingCameraActor.h
@@ -21,6 +21,7 @@ class UTextureRenderTarget2D;
 class UDEPRECATED_UOculusMR_Settings;
 class UOculusMR_State;
 class UMaterialInstanceDynamic;
+class UTexture2D;
 
 /**
 * The camera actor in the level that tracks the binded physical camera in game
