Adding more includes to prepare for upcoming header changes.

diff --git a/Engine/Plugins/Media/AvfMedia/Source/AvfMedia/Private/Player/AvfMediaVideoSampler.cpp b/Engine/Plugins/Media/AvfMedia/Source/AvfMedia/Private/Player/AvfMediaVideoSampler.cpp
--- a/Engine/Plugins/Media/AvfMedia/Source/AvfMedia/Private/Player/AvfMediaVideoSampler.cpp
+++ b/Engine/Plugins/Media/AvfMedia/Source/AvfMedia/Private/Player/AvfMediaVideoSampler.cpp
@@ -10,6 +10,7 @@
 #if WITH_ENGINE
 #include "RenderingThread.h"
 #include "RHI.h"
+#include "DataDrivenShaderPlatformInfo.h"
 #include "MediaShaders.h"
 #include "StaticBoundShaderState.h"
 #include "Misc/ScopeLock.h"
diff --git a/Engine/Plugins/Media/ElectraUtil/Source/ElectraSamples/Private/ElectraTextureSampleAVF.cpp b/Engine/Plugins/Media/ElectraUtil/Source/ElectraSamples/Private/ElectraTextureSampleAVF.cpp
--- a/Engine/Plugins/Media/ElectraUtil/Source/ElectraSamples/Private/ElectraTextureSampleAVF.cpp
+++ b/Engine/Plugins/Media/ElectraUtil/Source/ElectraSamples/Private/ElectraTextureSampleAVF.cpp
@@ -9,6 +9,7 @@
 #if WITH_ENGINE
 #include "RenderingThread.h"
 #include "RHI.h"
+#include "DataDrivenShaderPlatformInfo.h"
 #include "RHIStaticStates.h"
 #include "MediaShaders.h"
 #else
diff --git a/Engine/Plugins/Runtime/Oculus/OculusVR/Source/OculusMR/Private/OculusMR_CastingCameraActor.cpp b/Engine/Plugins/Runtime/Oculus/OculusVR/Source/OculusMR/Private/OculusMR_CastingCameraActor.cpp
--- a/Engine/Plugins/Runtime/Oculus/OculusVR/Source/OculusMR/Private/OculusMR_CastingCameraActor.cpp
+++ b/Engine/Plugins/Runtime/Oculus/OculusVR/Source/OculusMR/Private/OculusMR_CastingCameraActor.cpp
@@ -25,6 +25,7 @@
 #include "Materials/MaterialInstanceDynamic.h"
 #include "VRNotificationsComponent.h"
 #include "RenderUtils.h"
+#include "DataDrivenShaderPlatformInfo.h"
 #include "AudioDevice.h"
 
 #define LOCTEXT_NAMESPACE "OculusMR_CastingCameraActor"
diff --git a/Engine/Source/Runtime/Apple/MetalRHI/Private/MetalDynamicRHI_Shaders.cpp b/Engine/Source/Runtime/Apple/MetalRHI/Private/MetalDynamicRHI_Shaders.cpp
--- a/Engine/Source/Runtime/Apple/MetalRHI/Private/MetalDynamicRHI_Shaders.cpp
+++ b/Engine/Source/Runtime/Apple/MetalRHI/Private/MetalDynamicRHI_Shaders.cpp
@@ -8,6 +8,7 @@
 #include "MetalRHIPrivate.h"
 #include "MetalShaderTypes.h"
 #include "Shaders/MetalShaderLibrary.h"
+#include "DataDrivenShaderPlatformInfo.h"
 
 
 //------------------------------------------------------------------------------
diff --git a/Engine/Source/Runtime/Apple/MetalRHI/Private/MetalRenderPass.cpp b/Engine/Source/Runtime/Apple/MetalRHI/Private/MetalRenderPass.cpp
--- a/Engine/Source/Runtime/Apple/MetalRHI/Private/MetalRenderPass.cpp
+++ b/Engine/Source/Runtime/Apple/MetalRHI/Private/MetalRenderPass.cpp
@@ -13,6 +13,7 @@
 #include "MetalCommandBuffer.h"
 #include "MetalProfiler.h"
 #include "MetalFrameAllocator.h"
+#include "DataDrivenShaderPlatformInfo.h"
 
 #pragma mark - Private Console Variables -
 
