Fixing include paths inside Oculus plugin.

diff --git a/Engine/Plugins/Runtime/Oculus/OculusVR/Source/OculusHMD/Private/OculusHMD_Splash.cpp b/Engine/Plugins/Runtime/Oculus/OculusVR/Source/OculusHMD/Private/OculusHMD_Splash.cpp
--- a/Engine/Plugins/Runtime/Oculus/OculusVR/Source/OculusHMD/Private/OculusHMD_Splash.cpp
+++ b/Engine/Plugins/Runtime/Oculus/OculusVR/Source/OculusHMD/Private/OculusHMD_Splash.cpp
@@ -8,7 +8,7 @@
 #include "RenderingThread.h"
 #include "Misc/ScopeLock.h"
 #include "OculusHMDRuntimeSettings.h"
-#include "StereoLayerFunctionLibrary.h"
+#include "Kismet/StereoLayerFunctionLibrary.h"
 #if PLATFORM_ANDROID
 #include "Android/AndroidJNI.h"
 #include "Android/AndroidApplication.h"
diff --git a/Engine/Plugins/Runtime/Oculus/OculusVR/Source/OculusMR/Private/OculusMRModule.cpp b/Engine/Plugins/Runtime/Oculus/OculusVR/Source/OculusMR/Private/OculusMRModule.cpp
--- a/Engine/Plugins/Runtime/Oculus/OculusVR/Source/OculusMR/Private/OculusMRModule.cpp
+++ b/Engine/Plugins/Runtime/Oculus/OculusVR/Source/OculusMR/Private/OculusMRModule.cpp
@@ -6,7 +6,7 @@
 #include "IXRTrackingSystem.h"
 #include "StereoRendering.h"
 #include "StereoRenderTargetManager.h"
-#include "SceneCaptureComponent2D.h"
+#include "Components/SceneCaptureComponent2D.h"
 #include "Engine/TextureRenderTarget2D.h"
 #include "EngineUtils.h"
 //#include "PostProcess/SceneRenderTargets.h"
