Removing some commented out includes

diff --git a/Engine/Plugins/Experimental/SimpleHMD/Source/SimpleHMD/Private/SimpleHMD.cpp b/Engine/Plugins/Experimental/SimpleHMD/Source/SimpleHMD/Private/SimpleHMD.cpp
--- a/Engine/Plugins/Experimental/SimpleHMD/Source/SimpleHMD/Private/SimpleHMD.cpp
+++ b/Engine/Plugins/Experimental/SimpleHMD/Source/SimpleHMD/Private/SimpleHMD.cpp
@@ -6,7 +6,6 @@
 #include "ISimpleHMDPlugin.h"
 #include "PostProcess/PostProcessHMD.h"
 #include "GameFramework/WorldSettings.h"
-//#include "ScenePrivate.h"
 #include "SceneView.h"
 
 //---------------------------------------------------
diff --git a/Engine/Plugins/Runtime/Oculus/OculusVR/Source/OculusHMD/Private/OculusHMD_Layer.cpp b/Engine/Plugins/Runtime/Oculus/OculusVR/Source/OculusHMD/Private/OculusHMD_Layer.cpp
--- a/Engine/Plugins/Runtime/Oculus/OculusVR/Source/OculusHMD/Private/OculusHMD_Layer.cpp
+++ b/Engine/Plugins/Runtime/Oculus/OculusVR/Source/OculusHMD/Private/OculusHMD_Layer.cpp
@@ -4,9 +4,6 @@
 
 #if OCULUS_HMD_SUPPORTED_PLATFORMS
 //#include "MediaTexture.h"
-//#include "ScreenRendering.h"
-//#include "ScenePrivate.h"
-//#include "PostProcess/SceneFilterRendering.h"
 #include "Engine/GameEngine.h"
 #include "Engine/Texture2D.h"
 #include "HeadMountedDisplayTypes.h" // for LogHMD
@@ -15,7 +12,6 @@
 #include "OculusHMD.h"
 #include "OculusHMDModule.h"
 #include "OculusHMDPrivate.h"
-//#include "PostProcess/SceneRenderTargets.h"
 #include "SceneUtils.h"
 #include "UObject/ConstructorHelpers.h"
 #include "XRThreadUtils.h"
diff --git a/Engine/Plugins/Runtime/OpenXR/Source/OpenXRHMD/Private/OpenXRHMD.cpp b/Engine/Plugins/Runtime/OpenXR/Source/OpenXRHMD/Private/OpenXRHMD.cpp
--- a/Engine/Plugins/Runtime/OpenXR/Source/OpenXRHMD/Private/OpenXRHMD.cpp
+++ b/Engine/Plugins/Runtime/OpenXR/Source/OpenXRHMD/Private/OpenXRHMD.cpp
@@ -16,7 +16,6 @@
 #include "Engine/Engine.h"
 #include "GameFramework/PlayerController.h"
 #include "Engine/LocalPlayer.h"
-//#include "SceneRendering.h"
 #include "PostProcess/PostProcessHMD.h"
 #include "GameFramework/WorldSettings.h"
 #include "Misc/CString.h"
diff --git a/Engine/Plugins/Runtime/nDisplay/Source/DisplayClusterLightCardEditorShaders/Private/DisplayClusterMeshProjectionRenderer.cpp b/Engine/Plugins/Runtime/nDisplay/Source/DisplayClusterLightCardEditorShaders/Private/DisplayClusterMeshProjectionRenderer.cpp
--- a/Engine/Plugins/Runtime/nDisplay/Source/DisplayClusterLightCardEditorShaders/Private/DisplayClusterMeshProjectionRenderer.cpp
+++ b/Engine/Plugins/Runtime/nDisplay/Source/DisplayClusterLightCardEditorShaders/Private/DisplayClusterMeshProjectionRenderer.cpp
@@ -19,8 +19,6 @@
 #include "MaterialDomain.h"
 #include "Materials/Material.h"
 #include "DataDrivenShaderPlatformInfo.h"
-//#include "PostProcess/SceneFilterRendering.h"
-//#include "ScenePrivate.h"
 #include "ScreenPass.h"
 
 //////////////////////////////////////////////////////////////////////////
